Adding more informative error if swarm agent attempts to connect to a PC which name consists only of numbers (TTP# 338762)

diff --git a/Engine/Source/Programs/UnrealSwarm/Agent/Connections.cs b/Engine/Source/Programs/UnrealSwarm/Agent/Connections.cs
--- a/Engine/Source/Programs/UnrealSwarm/Agent/Connections.cs
+++ b/Engine/Source/Programs/UnrealSwarm/Agent/Connections.cs
@@ -779,6 +779,28 @@ namespace Agent
 
 		///////////////////////////////////////////////////////////////////////
 
+		/// <summary>
+		/// Validates host name. Currently checks if the name contains only numeric characters
+		/// Throws an exception if the name is invalid
+		/// </summary>
+		public void ValidateHostName(string HostName)
+		{
+			bool bNumericOnly = true;
+			foreach (char C in HostName)
+			{
+				if (!Char.IsDigit(C))
+				{
+					bNumericOnly = false;
+				}
+			}
+			if (bNumericOnly)
+			{
+				string Message = String.Format("Host name {0} contains only numbers and it may not be possible to connect to it", HostName);
+				Log(EVerbosityLevel.Informative, ELogColour.Red, " ......... " + Message);
+				throw new Exception(Message);
+			}
+		}
+
 		/**
 		 * Initializes the coordinator interface
 		 */
@@ -793,6 +815,8 @@ namespace Agent
 				string CoordinatorObjectURL = String.Format( "tcp://127.0.0.1:{0}/SwarmCoordinator",
 					Properties.Settings.Default.CoordinatorRemotingPort );
 #else
+				ValidateHostName(AgentApplication.Options.CoordinatorRemotingHost);
+
 				string CoordinatorObjectURL = String.Format( "tcp://{0}:{1}/SwarmCoordinator",
 					AgentApplication.Options.CoordinatorRemotingHost,
 					Properties.Settings.Default.CoordinatorRemotingPort );
@@ -1188,6 +1212,8 @@ namespace Agent
 							Debug.Assert( RequestingAgentInfo.Configuration.ContainsKey( "IPAddress" ) );
 							string RequestingAgentIPAddress = RequestingAgentInfo.Configuration["IPAddress"].ToString();
 
+							ValidateHostName(RequestingAgentIPAddress);
+
 							// Now, ping the Agent
 							if( PingRemoteHost( RequestingAgentName, RequestingAgentIPAddress ) )
 							{
@@ -1384,8 +1410,19 @@ namespace Agent
 
 			Log( EVerbosityLevel.Verbose, ELogColour.Green, "[Connect] Trying to open a remote connection to " + RemoteAgentName + " at " + RemoteAgentIPAddress );
 
+			bool bHostNameValid = true;
+			try
+			{
+				ValidateHostName(RemoteAgentIPAddress);
+			}
+			catch (Exception)
+			{
+				ErrorCode = Constants.ERROR_EXCEPTION;
+				bHostNameValid = false;
+			}
+
 			// First, ping the remote host to make sure we can talk to it
-			if( PingRemoteHost( RemoteAgentName, RemoteAgentIPAddress ) )
+			if (bHostNameValid && PingRemoteHost(RemoteAgentName, RemoteAgentIPAddress))
 			{
 				// Get the new unique handle that this connection will be known by
 				Int32 NewConnectionHandle = GetUniqueHandle();
