UnrealSwarm: Fix for "Pushing the channel has failed!" error.

The problem was that when optional file was missing, swarm tried to download it until success, without limited number of retries. Added retries parameter and set it to 5.

Also added some more informative error messages.

More info about the bug here: https://answers.unrealengine.com/questions/92329/pushing-the-channel-has-failed.html.

diff --git a/Engine/Source/Programs/UnrealSwarm/Agent/Channels.cs b/Engine/Source/Programs/UnrealSwarm/Agent/Channels.cs
--- a/Engine/Source/Programs/UnrealSwarm/Agent/Channels.cs
+++ b/Engine/Source/Programs/UnrealSwarm/Agent/Channels.cs
@@ -784,9 +784,9 @@ namespace Agent
 		 * Pushes a local channel to the remote agent via the remote connection parameter
 		 * by calling SendChannel on the remote agent
 		 */
-		public bool PushChannel( RemoteConnection Remote, string ChannelName, AgentGuid JobGuid )
+		public bool PushChannel(RemoteConnection Remote, string ChannelName, AgentGuid JobGuid)
 		{
-            StartTiming( "PushChannel-Internal", true );
+			StartTiming("PushChannel-Internal", true);
 			bool bChannelTransferred = false;
 
 			// The remote connection's handle can be used for all interaction because
@@ -794,7 +794,7 @@ namespace Agent
 			Int32 ConnectionHandle = Remote.Handle;
 
 			EChannelFlags ChannelFlags = EChannelFlags.ACCESS_READ;
-			if( JobGuid == null )
+			if (JobGuid == null)
 			{
 				ChannelFlags |= EChannelFlags.TYPE_PERSISTENT;
 			}
@@ -809,48 +809,47 @@ namespace Agent
 			OpenInParameters["ChannelFlags"] = ChannelFlags;
 			Hashtable OpenOutParameters = null;
 
-			Int32 LocalChannelHandle = OpenChannel_1_0( ConnectionHandle, OpenInParameters, ref OpenOutParameters );
-			if( LocalChannelHandle >= 0 )
+			Int32 LocalChannelHandle = OpenChannel_1_0(ConnectionHandle, OpenInParameters, ref OpenOutParameters);
+			if (LocalChannelHandle >= 0)
 			{
 				try
 				{
 					string FullChannelName;
-					if( JobGuid == null )
+					if (JobGuid == null)
 					{
-						FullChannelName = Path.Combine( AgentApplication.Options.CacheFolder, ChannelName );
+						FullChannelName = Path.Combine(AgentApplication.Options.CacheFolder, ChannelName);
 					}
 					else
 					{
-						string AllJobsFolder = Path.Combine( AgentApplication.Options.CacheFolder, "Jobs" );
-						string ThisJobFolder = Path.Combine( AllJobsFolder, "Job-" + JobGuid.ToString() );
-						FullChannelName = Path.Combine( ThisJobFolder, ChannelName );
+						string AllJobsFolder = Path.Combine(AgentApplication.Options.CacheFolder, "Jobs");
+						string ThisJobFolder = Path.Combine(AllJobsFolder, "Job-" + JobGuid.ToString());
+						FullChannelName = Path.Combine(ThisJobFolder, ChannelName);
 					}
 
 					// Read the entire file into a byte stream
-					byte[] ChannelData = File.ReadAllBytes( FullChannelName );
+					byte[] ChannelData = File.ReadAllBytes(FullChannelName);
 
 					// Send the entire channel at once
-					bChannelTransferred = Remote.Interface.SendChannel( ConnectionHandle, ChannelName, ChannelData, JobGuid );
+					bChannelTransferred = Remote.Interface.SendChannel(ConnectionHandle, ChannelName, ChannelData, JobGuid);
 
 					// If the channel was transferred, track the number of bytes that actually moved across the network
-					if( bChannelTransferred )
+					if (bChannelTransferred)
 					{
-						FileInfo ChannelInfo = new FileInfo( FullChannelName );
+						FileInfo ChannelInfo = new FileInfo(FullChannelName);
 						Remote.NetworkBytesSent += ChannelInfo.Length;
-						if( Remote.Job != null )
+						if (Remote.Job != null)
 						{
 							Remote.Job.NetworkBytesSent += ChannelInfo.Length;
 						}
-						if( Remote.Parent != null )
+						if (Remote.Parent != null)
 						{
 							Remote.Parent.NetworkBytesSent += ChannelInfo.Length;
 						}
 					}
 				}
-				catch( Exception Ex )
+				catch (Exception Ex)
 				{
-                    Log( EVerbosityLevel.Informative, ELogColour.Red, "[PushChannel] Pushing the channel has failed!" );
-                    Log( EVerbosityLevel.Verbose, ELogColour.Red, "[PushChannel] Exception message: " + Ex.ToString() );
+					Log(EVerbosityLevel.Verbose, ELogColour.Red, "[PushChannel] Exception message: " + Ex.ToString());
 				}
 
 				Hashtable CloseInParameters = new Hashtable();
@@ -858,12 +857,23 @@ namespace Agent
 				CloseInParameters["ChannelHandle"] = LocalChannelHandle;
 				Hashtable CloseOutParameters = null;
 
-				CloseChannel_1_0( ConnectionHandle, CloseInParameters, ref CloseOutParameters );
+				CloseChannel_1_0(ConnectionHandle, CloseInParameters, ref CloseOutParameters);
 
-                Log( EVerbosityLevel.Verbose, ELogColour.Green, "[Channel] Successful channel push of " + ChannelName );
+				if (bChannelTransferred)
+				{
+					Log(EVerbosityLevel.Verbose, ELogColour.Green, "[Channel] Successful channel push of " + ChannelName);
+				}
+				else
+				{
+					Log(EVerbosityLevel.Informative, ELogColour.Red, string.Format("[PushChannel] Pushing the channel {0} has failed!", ChannelName));
+				}
+			}
+			else
+			{
+				Log(EVerbosityLevel.Informative, ELogColour.Red, string.Format("[PushChannel] Cannot open local channel {0}.", ChannelName));
 			}
 
-            StopTiming();
+			StopTiming();
 			return bChannelTransferred;
 		}
 
@@ -961,12 +971,19 @@ namespace Agent
 		}
 
 		/**
-		 * Pulls a remote channel from the remote agent via the remote connection parameter
+		 * Pulls a remote channel from the remote agent via the remote connection parameter.
+		 * 
+		 * @param Remote Remote connection.
+		 * @param ChannelName The name of the file to pull.
+		 * @param JobGuid A guid of the job.
+		 * @param RetriesOnFailure How many times should the pull fail until return a failure.
+		 * 
+		 * @returns True on success. False otherwise.
 		 */
-		public bool PullChannel( RemoteConnection Remote, string ChannelName, AgentGuid JobGuid )
+		public bool PullChannel(RemoteConnection Remote, string ChannelName, AgentGuid JobGuid, int RetriesOnFailure = int.MaxValue)
 		{
-            StartTiming( "PullChannel-Internal", true );
-            
+			StartTiming("PullChannel-Internal", true);
+
 			// The remote connection's handle can be used for all interaction because
 			// it has the same meaning on both ends of the connection
 			Int32 ConnectionHandle = Remote.Handle;
@@ -974,21 +991,29 @@ namespace Agent
 			// Request the file, which will push it back, and keep doing so until we get it
 			// or until the connection is dead, which ever comes first
 			bool bChannelTransferred = false;
-			while( ( Remote.Interface.IsAlive() ) &&
-				   ( bChannelTransferred == false ) )
+			int TryId = 0;
+			while ((Remote.Interface.IsAlive()) &&
+				   (bChannelTransferred == false) &&
+				   (TryId < RetriesOnFailure))
 			{
 				try
 				{
-					bChannelTransferred = Remote.Interface.RequestChannel( ConnectionHandle, ChannelName, JobGuid );
+					bChannelTransferred = Remote.Interface.RequestChannel(ConnectionHandle, ChannelName, JobGuid);
 				}
-				catch( Exception Ex )
+				catch (Exception Ex)
 				{
-					Log( EVerbosityLevel.Informative, ELogColour.Red, "[PullChannel] Pulling the channel has failed!" );
-					Log( EVerbosityLevel.Verbose, ELogColour.Red, "[PullChannel] Exception message: " + Ex.ToString() );
+					Log(EVerbosityLevel.Verbose, ELogColour.Red, "[PullChannel] Exception message: " + Ex.ToString());
 				}
+
+				if(!bChannelTransferred)
+				{
+					Log(EVerbosityLevel.Informative, ELogColour.Red, string.Format("[PullChannel] Pulling the channel {0} has failed! Retry {1} of {2}.", ChannelName, TryId + 1, RetriesOnFailure));
+				}
+
+				++TryId;
 			}
 
-            StopTiming();
+			StopTiming();
 			return bChannelTransferred;
 		}
 
diff --git a/Engine/Source/Programs/UnrealSwarm/Agent/Jobs.cs b/Engine/Source/Programs/UnrealSwarm/Agent/Jobs.cs
--- a/Engine/Source/Programs/UnrealSwarm/Agent/Jobs.cs
+++ b/Engine/Source/Programs/UnrealSwarm/Agent/Jobs.cs
@@ -101,44 +101,44 @@ namespace Agent
 			return Constants.SUCCESS;
 		}
 
-		private bool RequestDependency( Connection RequestingConnection, string Dependency, bool bIsRequired )
+		private bool RequestDependency(Connection RequestingConnection, string Dependency, bool bIsRequired)
 		{
 			Hashtable InParameters = new Hashtable();
 			InParameters["Version"] = ESwarmVersionValue.VER_1_0;
 			InParameters["ChannelName"] = Dependency;
 			Hashtable OutParameters = null;
-			if( Manager.TestChannel( RequestingConnection.Handle, InParameters, ref OutParameters ) < 0 )
+			if (Manager.TestChannel(RequestingConnection.Handle, InParameters, ref OutParameters) < 0)
 			{
-				if( RequestingConnection is RemoteConnection )
+				if (RequestingConnection is RemoteConnection)
 				{
 					// For remote connection misses, try to pull the file from the instigator
-					Manager.Log( EVerbosityLevel.Verbose, ELogColour.Green, "[Job] Attempting to pull file from remote Agent: " + Dependency );
-					if( Manager.PullChannel( RequestingConnection as RemoteConnection, Dependency, null ) == false )
+					Manager.Log(EVerbosityLevel.Verbose, ELogColour.Green, "[Job] Attempting to pull file from remote Agent: " + Dependency);
+					if (Manager.PullChannel(RequestingConnection as RemoteConnection, Dependency, null, 5) == false)
 					{
-						if( bIsRequired )
+						if (bIsRequired)
 						{
-							Manager.Log( EVerbosityLevel.Informative, ELogColour.Red, "[Job] Error: Failed to pull a required file from remote Agent: " + Dependency );
+							Manager.Log(EVerbosityLevel.Informative, ELogColour.Red, "[Job] Error: Failed to pull a required file from remote Agent: " + Dependency);
 							return false;
 						}
 						else
 						{
-							Manager.Log( EVerbosityLevel.Verbose, ELogColour.Orange, "[Job] Warning: Failed to pull an optional file from remote Agent: " + Dependency );
+							Manager.Log(EVerbosityLevel.Verbose, ELogColour.Orange, "[Job] Warning: Failed to pull an optional file from remote Agent: " + Dependency);
 						}
 					}
 				}
 				else
 				{
-					Debug.Assert( RequestingConnection is LocalConnection );
+					Debug.Assert(RequestingConnection is LocalConnection);
 
-					if( bIsRequired )
+					if (bIsRequired)
 					{
 						// Always fail on local connection misses
-						Manager.Log( EVerbosityLevel.Informative, ELogColour.Red, "[Job] Error: Failed to find required file: " + Dependency );
+						Manager.Log(EVerbosityLevel.Informative, ELogColour.Red, "[Job] Error: Failed to find required file: " + Dependency);
 						return false;
 					}
 					else
 					{
-						Manager.Log( EVerbosityLevel.Verbose, ELogColour.Orange, "[Job] Warning: Failed to find an optional file: " + Dependency );
+						Manager.Log(EVerbosityLevel.Verbose, ELogColour.Orange, "[Job] Warning: Failed to find an optional file: " + Dependency);
 					}
 				}
 			}
