Swarm: Updated signing certificate binding in project

diff --git a/Engine/Source/Programs/UnrealSwarm/Agent/Agent.csproj b/Engine/Source/Programs/UnrealSwarm/Agent/Agent.csproj
--- a/Engine/Source/Programs/UnrealSwarm/Agent/Agent.csproj
+++ b/Engine/Source/Programs/UnrealSwarm/Agent/Agent.csproj
@@ -15,7 +15,7 @@
     <SccAuxPath>SAK</SccAuxPath>
     <SccProvider>SAK</SccProvider>
     <IsWebBootstrapper>false</IsWebBootstrapper>
-    <ManifestCertificateThumbprint>E10B5956C0FAFF8479C61A48C48EFDE555A18633</ManifestCertificateThumbprint>
+    <ManifestCertificateThumbprint>9545C2E66F84C8CF6586FABC3FF99AA367130F7D</ManifestCertificateThumbprint>
     <ManifestKeyFile>
     </ManifestKeyFile>
     <GenerateManifests>true</GenerateManifests>
