Swarm: Enabled auto-increment of publishing version

diff --git a/Engine/Source/Programs/UnrealSwarm/Agent/Agent.csproj b/Engine/Source/Programs/UnrealSwarm/Agent/Agent.csproj
--- a/Engine/Source/Programs/UnrealSwarm/Agent/Agent.csproj
+++ b/Engine/Source/Programs/UnrealSwarm/Agent/Agent.csproj
@@ -40,8 +40,8 @@
     <CreateWebPageOnPublish>true</CreateWebPageOnPublish>
     <WebPage>publish.htm</WebPage>
     <OpenBrowserOnPublish>false</OpenBrowserOnPublish>
-    <ApplicationRevision>3</ApplicationRevision>
-    <ApplicationVersion>1.7.2891.0</ApplicationVersion>
+    <ApplicationRevision>0</ApplicationRevision>
+    <ApplicationVersion>1.7.5000.%2a</ApplicationVersion>
     <UseApplicationTrust>false</UseApplicationTrust>
     <PublishWizardCompleted>true</PublishWizardCompleted>
     <BootstrapperEnabled>true</BootstrapperEnabled>
