Swarm: Checked in latest publish version

diff --git a/Engine/Source/Programs/UnrealSwarm/Agent/Agent.csproj b/Engine/Source/Programs/UnrealSwarm/Agent/Agent.csproj
--- a/Engine/Source/Programs/UnrealSwarm/Agent/Agent.csproj
+++ b/Engine/Source/Programs/UnrealSwarm/Agent/Agent.csproj
@@ -40,7 +40,7 @@
     <CreateWebPageOnPublish>true</CreateWebPageOnPublish>
     <WebPage>publish.htm</WebPage>
     <OpenBrowserOnPublish>false</OpenBrowserOnPublish>
-    <ApplicationRevision>0</ApplicationRevision>
+    <ApplicationRevision>1</ApplicationRevision>
     <ApplicationVersion>1.7.5000.%2a</ApplicationVersion>
     <UseApplicationTrust>false</UseApplicationTrust>
     <PublishWizardCompleted>true</PublishWizardCompleted>
