UnrealSwarm: Fix a problem found on AnswerHub that on some machines was exposing: "Error: A 32 bit processes cannot access modules of a 64 bit process."

diff --git a/Engine/Source/Programs/UnrealSwarm/Agent/Connections.cs b/Engine/Source/Programs/UnrealSwarm/Agent/Connections.cs
--- a/Engine/Source/Programs/UnrealSwarm/Agent/Connections.cs
+++ b/Engine/Source/Programs/UnrealSwarm/Agent/Connections.cs
@@ -17,6 +17,7 @@ using System.Runtime.Remoting;
 using System.Runtime.Remoting.Channels;
 using System.Runtime.Remoting.Channels.Tcp;
 using System.Threading;
+using System.Linq;
 
 using AgentInterface;
 using SwarmCoordinatorInterface;
@@ -855,12 +856,7 @@ namespace Agent
 		 */
 		public void SetCurrentDirectoryByProcessID (int ProcessID)
 		{
-			if (Environment.OSVersion.Platform != PlatformID.Unix) { // Mac is "Unix" to Mono for compatiblity reasons
-				Process ProcessObject = Process.GetProcessById (ProcessID);
-				Environment.CurrentDirectory = Path.GetDirectoryName (ProcessObject.MainModule.FileName); 
-			}
-			else
-				Environment.CurrentDirectory = Path.GetFullPath(Path.GetDirectoryName(new Uri(Assembly.GetExecutingAssembly().CodeBase).AbsolutePath) + "/../Mac/");
+			Environment.CurrentDirectory = GetProcessPathById(ProcessID);
 		}
 
 		/**
@@ -2501,6 +2497,31 @@ namespace Agent
 				CoordinatorResponding = true;
 			}
 		}
+
+		private string GetProcessPathById(int ProcessID)
+		{
+			if (Environment.OSVersion.Platform == PlatformID.Unix)
+			{ // Mac is "Unix" to Mono for compatiblity reasons
+				return Path.GetFullPath(Path.GetDirectoryName(new Uri(Assembly.GetExecutingAssembly().CodeBase).AbsolutePath) + "/../Mac/");
+			}
+
+			var WMIQuery = string.Format("SELECT ExecutablePath FROM Win32_Process WHERE ProcessId = {0}", ProcessID);
+			string ExecPath;
+
+			using(var WMISearcher = new ManagementObjectSearcher(WMIQuery))
+			using(var Results = WMISearcher.Get())
+			{
+				if (Results.Count == 0)
+				{
+					throw new InvalidOperationException(string.Format("Couldn't find the process with ID {0}.", ProcessID));
+				}
+
+				ExecPath = Results.Cast<ManagementObject>().First().GetPropertyValue("ExecutablePath") as string;
+			}
+
+			return Path.GetDirectoryName(ExecPath);
+		}
+
 	}
 }
 
