Fix for cores number limitation in Lightmass using UnrealSwarm.

diff --git a/Engine/Config/BaseLightmass.ini b/Engine/Config/BaseLightmass.ini
--- a/Engine/Config/BaseLightmass.ini
+++ b/Engine/Config/BaseLightmass.ini
@@ -18,7 +18,6 @@ DefaultStaticMeshLightingRes=32
 bAllowCropping=False
 bGarbageCollectAfterExport=True
 bRebuildDirtyGeometryForLighting=True
-NumUnusedLightmassThreads=2
 
 [DevOptions.StaticLightingSceneConstants]
 StaticLightingLevelScale=1
diff --git a/Engine/Source/Editor/UnrealEd/Private/Lightmass/Lightmass.cpp b/Engine/Source/Editor/UnrealEd/Private/Lightmass/Lightmass.cpp
--- a/Engine/Source/Editor/UnrealEd/Private/Lightmass/Lightmass.cpp
+++ b/Engine/Source/Editor/UnrealEd/Private/Lightmass/Lightmass.cpp
@@ -2586,22 +2586,6 @@ bool FLightmassProcessor::BeginRun()
 	{
 		CommandLineParameters += TEXT(" -stats");
 	}
-	
-	int32 NumUnusedLightmassThreads;
-	verify(GConfig->GetInt(TEXT("DevOptions.StaticLighting"), TEXT("NumUnusedLightmassThreads"), NumUnusedLightmassThreads, GLightmassIni));
-
-	const int32 NumVirtualCores = FPlatformMisc::NumberOfCoresIncludingHyperthreads();
-	int32 NumThreads = NumVirtualCores - NumUnusedLightmassThreads;
-
-	// On machines with few cores, each core will have a massive impact on build time, so we prioritize build latency over editor performance during the build
-	if (NumVirtualCores <= 4)
-	{
-		NumThreads = NumVirtualCores - 1;
-	}
-
-	NumThreads = FMath::Max(1, NumThreads);
-
-	CommandLineParameters += FString::Printf(TEXT(" -numthreads %i"), NumThreads);
 
 	NSwarm::FJobSpecification JobSpecification32, JobSpecification64;
 	if ( !bUse64bitProcess )
diff --git a/Engine/Source/Programs/UnrealSwarm/Agent/SettableOptions.cs b/Engine/Source/Programs/UnrealSwarm/Agent/SettableOptions.cs
--- a/Engine/Source/Programs/UnrealSwarm/Agent/SettableOptions.cs
+++ b/Engine/Source/Programs/UnrealSwarm/Agent/SettableOptions.cs
@@ -21,20 +21,21 @@ namespace Agent
 {
 	public enum EOptionsVersion
 	{
-		OriginalVersion					= 3,
-		NewBarColors					= 4,
-		CacheManagement					= 5,
-		NewStartupColor					= 6,
-		AdjustedProgressionColors		= 7,
-		TweaksToTheNumberOfThreads		= 8,
-		FixTheOptionsReset				= 9,
-		DeployNewFarmMachineNames		= 10,
-		BumpTheJobsToKeepCount			= 11,
-		MoveTheDefaultCacheFolder		= 12,
-		ChangeCoordinatorLocation		= 13,
-		AddPerformanceMonitoringDisable	= 14,
-
-		Current							= AddPerformanceMonitoringDisable
+		OriginalVersion						= 3,
+		NewBarColors						= 4,
+		CacheManagement						= 5,
+		NewStartupColor						= 6,
+		AdjustedProgressionColors			= 7,
+		TweaksToTheNumberOfThreads			= 8,
+		FixTheOptionsReset					= 9,
+		DeployNewFarmMachineNames			= 10,
+		BumpTheJobsToKeepCount				= 11,
+		MoveTheDefaultCacheFolder			= 12,
+		ChangeCoordinatorLocation			= 13,
+		AddPerformanceMonitoringDisable		= 14,
+		MakeLocalJobsProcessorsCountLimited = 15,
+
+		Current								= MakeLocalJobsProcessorsCountLimited
 	};
 
 	public class EnumedCollectionPropertyDescriptor<E, T> : PropertyDescriptor
@@ -628,7 +629,7 @@ namespace Agent
 			set { PrivateLocalJobsDefaultProcessorCount = value; }
 		}
 		[XmlAttribute]
-		private int PrivateLocalJobsDefaultProcessorCount = Environment.ProcessorCount;
+		private int PrivateLocalJobsDefaultProcessorCount = GetDefaultLocalJobsDefaultProcessorCount();
 
 		[CategoryAttribute( "Local Performance Settings" )]
 		[DescriptionAttribute( "The default priority for locally initiated jobs" )]
@@ -648,11 +649,7 @@ namespace Agent
 			set { PrivateRemoteJobsDefaultProcessorCount = value; }
 		}
 		[XmlAttribute]
-#if !__MonoCS__
-		private int PrivateRemoteJobsDefaultProcessorCount = Math.Max( 1, ( ApplicationDeployment.IsNetworkDeployed ? Environment.ProcessorCount : Environment.ProcessorCount - 1 ) );
-#else
-		private int PrivateRemoteJobsDefaultProcessorCount = Math.Max( 1, Environment.ProcessorCount - 1 );
-#endif
+		private int PrivateRemoteJobsDefaultProcessorCount = GetDefaultRemoteJobsDefaultProcessorCount();
 
 		[CategoryAttribute( "Local Performance Settings" )]
 		[DescriptionAttribute( "The default priority for remotely initiated jobs" )]
@@ -725,6 +722,45 @@ namespace Agent
 		[XmlAttribute]
 		private int LocalOptionsVersion = ( int )EOptionsVersion.Current;
 
+		/**
+		 * Gets allowed processors count on this machine for local tasks.
+		 *
+		 * This value can be lass than logical processors count as we don't
+		 * want to block the editor on user's machine.
+		 *
+		 * @returns Allowed processors count.
+		 */
+		private static Int32 GetAllowedProcessorCount()
+		{
+			const Int32 NumUnusedSwarmThreads = 2;
+			
+			Int32 NumVirtualCores = Environment.ProcessorCount;
+			Int32 NumThreads = NumVirtualCores - NumUnusedSwarmThreads;
+
+			// On machines with few cores, each core will have a massive impact on build times, so we prioritize build latency over editor performance during the build
+			if (NumVirtualCores <= 4)
+			{
+				NumThreads = NumVirtualCores - 1;
+			}
+
+			return Math.Max(1, NumThreads);
+		}
+
+		private static Int32 GetDefaultLocalJobsDefaultProcessorCount()
+		{
+			return
+#if !__MonoCS__
+				ApplicationDeployment.IsNetworkDeployed ? Environment.ProcessorCount : GetAllowedProcessorCount();
+#else
+				GetAllowedProcessorCount();
+#endif
+		}
+
+		private static Int32 GetDefaultRemoteJobsDefaultProcessorCount()
+		{
+			return Environment.ProcessorCount;
+		}
+
 		/*
 		 * Set up any required default states
 		 */
@@ -732,14 +768,13 @@ namespace Agent
 		{
 			ReplayLastDistribution = false;
             LocalEnableLocalPerformanceMonitoring = true;
-			LocalJobsDefaultProcessorCount = Environment.ProcessorCount;
+
+			LocalJobsDefaultProcessorCount = GetDefaultLocalJobsDefaultProcessorCount();
 			LocalJobsDefaultProcessPriority = AgentProcessPriorityClass.BelowNormal;
-#if !__MonoCS__
-			RemoteJobsDefaultProcessorCount = Math.Max( 1, ( ApplicationDeployment.IsNetworkDeployed ? Environment.ProcessorCount : Environment.ProcessorCount - 1 ) );
-#else
-			RemoteJobsDefaultProcessorCount = Math.Max( 1, Environment.ProcessorCount - 1 );
-#endif
+
+			RemoteJobsDefaultProcessorCount = GetDefaultRemoteJobsDefaultProcessorCount();
 			RemoteJobsDefaultProcessPriority = AgentProcessPriorityClass.Idle;
+
 			JobExecutableTimeout = 5;
 			RemoteAgentTimeout = 30;
 #if !__MonoCS__
