Fix multiple adapters IP resolve problem in UnrealSwarm.

diff --git a/Engine/Source/Programs/UnrealSwarm/Agent/Connections.cs b/Engine/Source/Programs/UnrealSwarm/Agent/Connections.cs
--- a/Engine/Source/Programs/UnrealSwarm/Agent/Connections.cs
+++ b/Engine/Source/Programs/UnrealSwarm/Agent/Connections.cs
@@ -21,6 +21,7 @@ using System.Linq;
 
 using AgentInterface;
 using SwarmCoordinatorInterface;
+using CommonUtils;
 
 namespace Agent
 {
@@ -2107,32 +2108,8 @@ namespace Agent
 						AgentInfoUpdate.Configuration["WorkingFor"] = WorkingFor;
 
 #if !__MonoCS__ // @todo Mac
-						// Send our IP address for others to use
-						ManagementObjectSearcher ObjectSearcher = new ManagementObjectSearcher( "SELECT * FROM Win32_NetworkAdapterConfiguration WHERE IPEnabled = 'TRUE'" );
-						ManagementObjectCollection ObjectCollection = ObjectSearcher.Get();
-                        bool bFoundAddress = false;
-                        bool bFoundAddressWithDNS = false;
-
-						foreach( ManagementObject NextObject in ObjectCollection )
-						{
-							// Get the first IP addresses of the first IPEnabled network adapter
-							string[] AddressList = ( string[] )NextObject["IPAddress"];
-                            string DNSDomain = (string)NextObject["DNSDomain"];
-
-                            // Use the first address in the enumerator, or the first address with a DNS gateway if one exists
-                            // Checking the DNS gateway allows adapters connected to networks to win over adapters earlier in the enumeration that are on a private network, eg connected to console
-                            if (!bFoundAddress || !bFoundAddressWithDNS)
-							{
-                                AgentInfoUpdate.Configuration["IPAddress"] = IPAddress.Parse(AddressList[0]);
-
-                                bFoundAddress = true;
-
-                                if (DNSDomain != null && DNSDomain.Length > 0)
-                                {
-                                    bFoundAddressWithDNS = true;
-                                }
-							}
-						}
+						var NetworkInterface = NetworkUtils.GetBestInterface(IPAddress.Parse(AgentApplication.Options.CoordinatorRemotingHost));
+						AgentInfoUpdate.Configuration["IPAddress"] = NetworkUtils.GetInterfaceIPv4Address(NetworkInterface);
 #else
 						AgentInfoUpdate.Configuration["IPAddress"] = "192.168.0.203";//.2.9";
 #endif
diff --git a/Engine/Source/Programs/UnrealSwarm/Agent/Properties/AssemblyInfo.cs b/Engine/Source/Programs/UnrealSwarm/Agent/Properties/AssemblyInfo.cs
--- a/Engine/Source/Programs/UnrealSwarm/Agent/Properties/AssemblyInfo.cs
+++ b/Engine/Source/Programs/UnrealSwarm/Agent/Properties/AssemblyInfo.cs
@@ -17,4 +17,4 @@ using System.Reflection;
 //      Build Number
 //      Revision
 //
-[assembly: AssemblyVersion( "1.7.2891.0" )]
+[assembly: AssemblyVersion( "1.7.2892.0" )]
diff --git a/Engine/Source/Programs/UnrealSwarm/CommonUtils/CommonUtils.csproj b/Engine/Source/Programs/UnrealSwarm/CommonUtils/CommonUtils.csproj
new file mode 100644
--- /dev/null
+++ b/Engine/Source/Programs/UnrealSwarm/CommonUtils/CommonUtils.csproj
@@ -0,0 +1,54 @@
+﻿<?xml version="1.0" encoding="utf-8"?>
+<Project ToolsVersion="12.0" DefaultTargets="Build" xmlns="http://schemas.microsoft.com/developer/msbuild/2003">
+  <Import Project="$(MSBuildExtensionsPath)\$(MSBuildToolsVersion)\Microsoft.Common.props" Condition="Exists('$(MSBuildExtensionsPath)\$(MSBuildToolsVersion)\Microsoft.Common.props')" />
+  <PropertyGroup>
+    <Configuration Condition=" '$(Configuration)' == '' ">Debug</Configuration>
+    <Platform Condition=" '$(Platform)' == '' ">AnyCPU</Platform>
+    <ProjectGuid>{65521250-B022-442F-B961-3AF5ADA02083}</ProjectGuid>
+    <OutputType>Library</OutputType>
+    <AppDesignerFolder>Properties</AppDesignerFolder>
+    <RootNamespace>CommonUtils</RootNamespace>
+    <AssemblyName>CommonUtils</AssemblyName>
+    <TargetFrameworkVersion>v4.0</TargetFrameworkVersion>
+    <FileAlignment>512</FileAlignment>
+    <TargetFrameworkProfile>Client</TargetFrameworkProfile>
+  </PropertyGroup>
+  <PropertyGroup Condition=" '$(Configuration)|$(Platform)' == 'Debug|AnyCPU' ">
+    <DebugSymbols>true</DebugSymbols>
+    <DebugType>full</DebugType>
+    <Optimize>false</Optimize>
+    <OutputPath>..\..\..\..\Binaries\DotNET\</OutputPath>
+    <DefineConstants>DEBUG;TRACE</DefineConstants>
+    <ErrorReport>prompt</ErrorReport>
+    <WarningLevel>4</WarningLevel>
+  </PropertyGroup>
+  <PropertyGroup Condition=" '$(Configuration)|$(Platform)' == 'Development|AnyCPU' ">
+    <DebugType>pdbonly</DebugType>
+    <Optimize>true</Optimize>
+    <OutputPath>..\..\..\..\Binaries\DotNET\</OutputPath>
+    <DefineConstants>TRACE</DefineConstants>
+    <ErrorReport>prompt</ErrorReport>
+    <WarningLevel>4</WarningLevel>
+  </PropertyGroup>
+  <ItemGroup>
+    <Reference Include="System" />
+    <Reference Include="System.Core" />
+    <Reference Include="System.Xml.Linq" />
+    <Reference Include="System.Data.DataSetExtensions" />
+    <Reference Include="Microsoft.CSharp" />
+    <Reference Include="System.Data" />
+    <Reference Include="System.Xml" />
+  </ItemGroup>
+  <ItemGroup>
+    <Compile Include="Utils.cs" />
+    <Compile Include="Properties\AssemblyInfo.cs" />
+  </ItemGroup>
+  <Import Project="$(MSBuildToolsPath)\Microsoft.CSharp.targets" />
+  <!-- To modify your build process, add your task inside one of the targets below and uncomment it. 
+       Other similar extension points exist, see Microsoft.Common.targets.
+  <Target Name="BeforeBuild">
+  </Target>
+  <Target Name="AfterBuild">
+  </Target>
+  -->
+</Project>
\ No newline at end of file
diff --git a/Engine/Source/Programs/UnrealSwarm/CommonUtils/Properties/AssemblyInfo.cs b/Engine/Source/Programs/UnrealSwarm/CommonUtils/Properties/AssemblyInfo.cs
new file mode 100644
--- /dev/null
+++ b/Engine/Source/Programs/UnrealSwarm/CommonUtils/Properties/AssemblyInfo.cs
@@ -0,0 +1,38 @@
+﻿// Copyright 1998-2015 Epic Games, Inc. All Rights Reserved.
+
+using System.Reflection;
+using System.Runtime.CompilerServices;
+using System.Runtime.InteropServices;
+
+// General Information about an assembly is controlled through the following 
+// set of attributes. Change these attribute values to modify the information
+// associated with an assembly.
+[assembly: AssemblyTitle("CommonUtils")]
+[assembly: AssemblyDescription("")]
+[assembly: AssemblyConfiguration("")]
+[assembly: AssemblyCompany("")]
+[assembly: AssemblyProduct("CommonUtils")]
+[assembly: AssemblyCopyright("Copyright ©  2014")]
+[assembly: AssemblyTrademark("")]
+[assembly: AssemblyCulture("")]
+
+// Setting ComVisible to false makes the types in this assembly not visible 
+// to COM components.  If you need to access a type in this assembly from 
+// COM, set the ComVisible attribute to true on that type.
+[assembly: ComVisible(false)]
+
+// The following GUID is for the ID of the typelib if this project is exposed to COM
+[assembly: Guid("67079ec4-8039-4160-ac89-e5da5582dc8c")]
+
+// Version information for an assembly consists of the following four values:
+//
+//      Major Version
+//      Minor Version 
+//      Build Number
+//      Revision
+//
+// You can specify all the values or you can default the Build and Revision Numbers 
+// by using the '*' as shown below:
+// [assembly: AssemblyVersion("1.0.*")]
+[assembly: AssemblyVersion("1.0.0.0")]
+[assembly: AssemblyFileVersion("1.0.0.0")]
diff --git a/Engine/Source/Programs/UnrealSwarm/CommonUtils/Utils.cs b/Engine/Source/Programs/UnrealSwarm/CommonUtils/Utils.cs
new file mode 100644
--- /dev/null
+++ b/Engine/Source/Programs/UnrealSwarm/CommonUtils/Utils.cs
@@ -0,0 +1,78 @@
+﻿// Copyright 1998-2015 Epic Games, Inc. All Rights Reserved.
+
+using System;
+using System.Collections.Generic;
+using System.IO;
+using System.Linq;
+using System.Net;
+using System.Net.NetworkInformation;
+using System.Runtime.InteropServices;
+using System.Runtime.Serialization.Formatters.Binary;
+using System.Text;
+using System.Threading.Tasks;
+
+namespace CommonUtils
+{
+	public static class NetworkUtils
+	{
+		[DllImport("iphlpapi.dll", SetLastError = true, EntryPoint = "GetBestInterface")]
+		static extern int WINAPI_GetBestInterface(UInt32 DestAddr, out UInt32 BestIfIndex);
+
+		public static NetworkInterface GetBestInterface(IPAddress Address)
+		{
+			UInt32 AddressAsUInt32 = BitConverter.ToUInt32(Address.GetAddressBytes(), 0);
+			UInt32 InterfaceIndex;
+
+			if (WINAPI_GetBestInterface(AddressAsUInt32, out InterfaceIndex) != 0)
+			{
+				throw new InvalidOperationException("Couldn't detect best network interface for that address.");
+			}
+
+			var Interfaces = NetworkInterface.GetAllNetworkInterfaces().Where(Interface =>
+			{
+				return Interface.Supports(NetworkInterfaceComponent.IPv4) && Interface.GetIPProperties().GetIPv4Properties().Index == InterfaceIndex;
+			});
+
+			if (Interfaces.Count() == 0)
+			{
+				throw new InvalidOperationException("Couldn't detect best network interface for that address.");
+			}
+
+			return Interfaces.First();
+		}
+
+		public static IPAddress GetInterfaceIPv4Address(NetworkInterface NetworkInterface)
+		{
+			var NetworkInterfaceProp = NetworkInterface.GetIPProperties();
+
+			foreach (var Unicast in NetworkInterfaceProp.UnicastAddresses)
+			{
+				if (Unicast.Address.AddressFamily == System.Net.Sockets.AddressFamily.InterNetwork)
+				{
+					return Unicast.Address;
+				}
+			}
+
+			throw new InvalidOperationException("Couldn't find interface IPv4 address.");
+		}
+	}
+
+	public static class ObjectUtils
+	{
+		public static T Duplicate<T>(T SourceObject)
+		{
+			if (!typeof(T).IsSerializable)
+			{
+				throw new InvalidOperationException("Only serializable types are supported.");
+			}
+
+			var Formatter = new BinaryFormatter();
+			using (var MemoryStream = new MemoryStream())
+			{
+				Formatter.Serialize(MemoryStream, SourceObject);
+				MemoryStream.Seek(0, SeekOrigin.Begin);
+				return (T)Formatter.Deserialize(MemoryStream);
+			}
+		}
+	}
+}
diff --git a/Engine/Source/Programs/UnrealSwarm/SwarmCoordinator/Coordinator.cs b/Engine/Source/Programs/UnrealSwarm/SwarmCoordinator/Coordinator.cs
--- a/Engine/Source/Programs/UnrealSwarm/SwarmCoordinator/Coordinator.cs
+++ b/Engine/Source/Programs/UnrealSwarm/SwarmCoordinator/Coordinator.cs
@@ -8,6 +8,7 @@ using System.Net;
 using System.Net.NetworkInformation;
 using System.Text;
 using SwarmCoordinatorInterface;
+using CommonUtils;
 
 namespace SwarmCoordinator
 {
@@ -127,29 +128,7 @@ namespace SwarmCoordinator
 				UpdatedInfo.Configuration["IPAddress"] = IPAddress.Loopback;
 				bool IPAddressUpdateSuccess = false;
 
-				if( UpdatedInfo.Name == Environment.MachineName )
-				{
-#if !__MonoCS__ // @todo Mac
-					// If it's this machine, do a simple lookup
-					ManagementObjectSearcher ObjectSearcher = new ManagementObjectSearcher( "SELECT * FROM Win32_NetworkAdapterConfiguration WHERE IPEnabled = 'TRUE'" );
-					ManagementObjectCollection ObjectCollection = ObjectSearcher.Get();
-					foreach( ManagementObject NextObject in ObjectCollection )
-					{
-						// Get the first IP addresses of the first IPEnabled network adapter
-						string[] AddressList = ( string[] )NextObject["IPAddress"];
-						if( AddressList.Length > 0 )
-						{
-							UpdatedInfo.Configuration["IPAddress"] = IPAddress.Parse( AddressList[0] );
-							IPAddressUpdateSuccess = true;
-							break;
-						}
-					}
-#else
-					UpdatedInfo.Configuration["IPAddress"] = "192.168.0.203"; //.2.9";
-					IPAddressUpdateSuccess = true;
-#endif
-				}
-				else
+				if( UpdatedInfo.Name != Environment.MachineName )
 				{
 					try
 					{
@@ -343,20 +322,28 @@ namespace SwarmCoordinator
 				return null;
 			}
 
+			AgentInfo RequestingAgent;
+			IPAddress MyAddressFromRequestingAgentPOV;
+
 			// If the requestor is Blocked, ignore the request
-			if( RequestedConfiguration.ContainsKey( "RequestingAgentName" ) )
+			if (!RequestedConfiguration.ContainsKey("RequestingAgentName"))
 			{
-				// If the requesting agent is not an approved state, block this request
-				string RequestingAgentName = RequestedConfiguration["RequestingAgentName"] as string;
-				if( ( Agents.ContainsKey( RequestingAgentName ) ) &&
-					( ( Agents[RequestingAgentName].State != AgentState.Available ) &&
-					  ( Agents[RequestingAgentName].State != AgentState.Working ) &&
-					  ( Agents[RequestingAgentName].State != AgentState.Busy ) ) )
-				{
-					return null;
-				}
+				return null;
+			}
+
+			// If the requesting agent is not an approved state, block this request
+			string RequestingAgentName = RequestedConfiguration["RequestingAgentName"] as string;
+			if ((Agents.TryGetValue(RequestingAgentName, out RequestingAgent) &&
+				((RequestingAgent.State != AgentState.Available) &&
+					(RequestingAgent.State != AgentState.Working) &&
+					(RequestingAgent.State != AgentState.Busy))))
+			{
+				return null;
 			}
 
+			var BestInt = NetworkUtils.GetBestInterface(RequestingAgent.Configuration["IPAddress"] as IPAddress);
+			MyAddressFromRequestingAgentPOV = NetworkUtils.GetInterfaceIPv4Address(BestInt);
+
 			// Extract the requested configuration
 			Version RequestedVersion = RequestedConfiguration["Version"] as Version;
 
@@ -390,23 +377,30 @@ namespace SwarmCoordinator
 					// If all conditions are met, we match
 					if( VersionMatch && StateMatch && GroupMatch )
 					{
-						AliveAgents.Add( Agent );
+						if (MyAddressFromRequestingAgentPOV != IPAddress.Loopback &&
+							(Agent.Configuration["IPAddress"] as IPAddress).Equals(IPAddress.Loopback))
+						{
+							var NewAgentInfo = ObjectUtils.Duplicate(Agent);
+
+							NewAgentInfo.Configuration["IPAddress"] = MyAddressFromRequestingAgentPOV;
+
+							AliveAgents.Add(NewAgentInfo);
+						}
+						else
+						{
+							AliveAgents.Add(Agent);
+						}
 					}
 				}
 
 				// If the requesting agent is asking for assignment, do it now that we have the right set
 				bool ValidRequestForAssignment = false;
-				string RequestingAgentName = "";
 				if( RequestedConfiguration.ContainsKey( "RequestAssignmentFor" ) )
 				{
 					// If the requesting agent is not an approved state, block this request
-					RequestingAgentName = RequestedConfiguration["RequestAssignmentFor"] as string;
-					if( Agents.ContainsKey( RequestingAgentName ) )
-					{
-						ValidRequestForAssignment = ( Agents[RequestingAgentName].State == AgentState.Available ) ||
-													( Agents[RequestingAgentName].State == AgentState.Working ) ||
-													( Agents[RequestingAgentName].State == AgentState.Busy );
-					}
+					ValidRequestForAssignment = (RequestingAgent.State == AgentState.Available) ||
+												(RequestingAgent.State == AgentState.Working) ||
+												(RequestingAgent.State == AgentState.Busy);
 				}
 
 				if( ValidRequestForAssignment )
@@ -420,7 +414,7 @@ namespace SwarmCoordinator
 					{
 						if( Agent.State == AgentState.Available )
 						{
-							Agent.Configuration["AssignedTo"] = RequestingAgentName;
+							Agent.Configuration["AssignedTo"] = RequestingAgent.Name;
 							Agent.Configuration["AssignedTime"] = DateTime.UtcNow;
 							NewlyAssignedCount++;
 						}
@@ -476,7 +470,7 @@ namespace SwarmCoordinator
 							LargestBin.RemoveAt( ReassignedAgentIndex );
 
 							// Reassigned and adjust the counts
-							ReassignedAgent.Configuration["AssignedTo"] = RequestingAgentName;
+							ReassignedAgent.Configuration["AssignedTo"] = RequestingAgent.Name;
 							ReassignedAgent.Configuration["AssignedTime"] = DateTime.UtcNow;
 							AlreadyAssignedCount--;
 							NewlyAssignedCount++;
