Missing references for CL 2399603.

diff --git a/Engine/Source/Programs/UnrealSwarm/Agent/Agent.csproj b/Engine/Source/Programs/UnrealSwarm/Agent/Agent.csproj
--- a/Engine/Source/Programs/UnrealSwarm/Agent/Agent.csproj
+++ b/Engine/Source/Programs/UnrealSwarm/Agent/Agent.csproj
@@ -145,6 +145,10 @@
       <Project>{6F18ADAC-2619-4C81-B80F-22927E56CA47}</Project>
       <Name>AgentInterface</Name>
     </ProjectReference>
+    <ProjectReference Include="..\CommonUtils\CommonUtils.csproj">
+      <Project>{65521250-b022-442f-b961-3af5ada02083}</Project>
+      <Name>CommonUtils</Name>
+    </ProjectReference>
     <ProjectReference Include="..\SwarmCoordinatorInterface\SwarmCoordinatorInterface.csproj">
       <Project>{4C299DD0-B968-49F2-894A-9F115C369E72}</Project>
       <Name>SwarmCoordinatorInterface</Name>
diff --git a/Engine/Source/Programs/UnrealSwarm/SwarmCoordinator/SwarmCoordinator.csproj b/Engine/Source/Programs/UnrealSwarm/SwarmCoordinator/SwarmCoordinator.csproj
--- a/Engine/Source/Programs/UnrealSwarm/SwarmCoordinator/SwarmCoordinator.csproj
+++ b/Engine/Source/Programs/UnrealSwarm/SwarmCoordinator/SwarmCoordinator.csproj
@@ -121,6 +121,10 @@
     </Compile>
   </ItemGroup>
   <ItemGroup>
+    <ProjectReference Include="..\CommonUtils\CommonUtils.csproj">
+      <Project>{65521250-b022-442f-b961-3af5ada02083}</Project>
+      <Name>CommonUtils</Name>
+    </ProjectReference>
     <ProjectReference Include="..\SwarmCoordinatorInterface\SwarmCoordinatorInterface.csproj">
       <Project>{4C299DD0-B968-49F2-894A-9F115C369E72}</Project>
       <Name>SwarmCoordinatorInterface</Name>
