UnrealSwarm: Adding required references to Mono projects.

diff --git a/Engine/Source/Programs/UnrealSwarm/Agent/Agent_Mono.csproj b/Engine/Source/Programs/UnrealSwarm/Agent/Agent_Mono.csproj
--- a/Engine/Source/Programs/UnrealSwarm/Agent/Agent_Mono.csproj
+++ b/Engine/Source/Programs/UnrealSwarm/Agent/Agent_Mono.csproj
@@ -144,6 +144,10 @@
       <Project>{6F18ADAC-2619-4C81-B80F-22927E56CA47}</Project>
       <Name>AgentInterface_Mono</Name>
     </ProjectReference>
+    <ProjectReference Include="..\CommonUtils\CommonUtils.csproj">
+      <Project>{65521250-b022-442f-b961-3af5ada02083}</Project>
+      <Name>CommonUtils</Name>
+    </ProjectReference>
     <ProjectReference Include="..\SwarmCoordinatorInterface\SwarmCoordinatorInterface_Mono.csproj">
       <Project>{4C299DD0-B968-49F2-894A-9F115C369E72}</Project>
       <Name>SwarmCoordinatorInterface_Mono</Name>
diff --git a/Engine/Source/Programs/UnrealSwarm/SwarmCoordinator/SwarmCoordinator_Mono.csproj b/Engine/Source/Programs/UnrealSwarm/SwarmCoordinator/SwarmCoordinator_Mono.csproj
--- a/Engine/Source/Programs/UnrealSwarm/SwarmCoordinator/SwarmCoordinator_Mono.csproj
+++ b/Engine/Source/Programs/UnrealSwarm/SwarmCoordinator/SwarmCoordinator_Mono.csproj
@@ -161,6 +161,10 @@
   </Target>
   -->
   <ItemGroup>
+    <ProjectReference Include="..\CommonUtils\CommonUtils.csproj">
+      <Project>{65521250-b022-442f-b961-3af5ada02083}</Project>
+      <Name>CommonUtils</Name>
+    </ProjectReference>
     <ProjectReference Include="..\SwarmCoordinatorInterface\SwarmCoordinatorInterface_Mono.csproj">
       <Project>{4C299DD0-B968-49F2-894A-9F115C369E72}</Project>
       <Name>SwarmCoordinatorInterface_Mono</Name>
