UnrealSwarm: Renamed new library to avoid future possible name conflicts and added this library to SwarmBuildProducts.

diff --git a/Engine/Source/Programs/AutomationTool/UE4Build.cs b/Engine/Source/Programs/AutomationTool/UE4Build.cs
--- a/Engine/Source/Programs/AutomationTool/UE4Build.cs
+++ b/Engine/Source/Programs/AutomationTool/UE4Build.cs
@@ -377,6 +377,7 @@ namespace AutomationTool
                 "SwarmCoordinator",
                 "SwarmCoordinatorInterface",
                 "SwarmInterface",
+				"SwarmCommonUtils"
             };
 		}
 
diff --git a/Engine/Source/Programs/UnrealSwarm/Agent/Agent.csproj b/Engine/Source/Programs/UnrealSwarm/Agent/Agent.csproj
--- a/Engine/Source/Programs/UnrealSwarm/Agent/Agent.csproj
+++ b/Engine/Source/Programs/UnrealSwarm/Agent/Agent.csproj
@@ -145,9 +145,9 @@
       <Project>{6F18ADAC-2619-4C81-B80F-22927E56CA47}</Project>
       <Name>AgentInterface</Name>
     </ProjectReference>
-    <ProjectReference Include="..\CommonUtils\CommonUtils.csproj">
+    <ProjectReference Include="..\SwarmCommonUtils\SwarmCommonUtils.csproj">
       <Project>{65521250-b022-442f-b961-3af5ada02083}</Project>
-      <Name>CommonUtils</Name>
+      <Name>SwarmCommonUtils</Name>
     </ProjectReference>
     <ProjectReference Include="..\SwarmCoordinatorInterface\SwarmCoordinatorInterface.csproj">
       <Project>{4C299DD0-B968-49F2-894A-9F115C369E72}</Project>
diff --git a/Engine/Source/Programs/UnrealSwarm/Agent/Agent_Mono.csproj b/Engine/Source/Programs/UnrealSwarm/Agent/Agent_Mono.csproj
--- a/Engine/Source/Programs/UnrealSwarm/Agent/Agent_Mono.csproj
+++ b/Engine/Source/Programs/UnrealSwarm/Agent/Agent_Mono.csproj
@@ -144,9 +144,9 @@
       <Project>{6F18ADAC-2619-4C81-B80F-22927E56CA47}</Project>
       <Name>AgentInterface_Mono</Name>
     </ProjectReference>
-    <ProjectReference Include="..\CommonUtils\CommonUtils.csproj">
+    <ProjectReference Include="..\SwarmCommonUtils\SwarmCommonUtils.csproj">
       <Project>{65521250-b022-442f-b961-3af5ada02083}</Project>
-      <Name>CommonUtils</Name>
+      <Name>SwarmCommonUtils</Name>
     </ProjectReference>
     <ProjectReference Include="..\SwarmCoordinatorInterface\SwarmCoordinatorInterface_Mono.csproj">
       <Project>{4C299DD0-B968-49F2-894A-9F115C369E72}</Project>
diff --git a/Engine/Source/Programs/UnrealSwarm/Agent/Connections.cs b/Engine/Source/Programs/UnrealSwarm/Agent/Connections.cs
--- a/Engine/Source/Programs/UnrealSwarm/Agent/Connections.cs
+++ b/Engine/Source/Programs/UnrealSwarm/Agent/Connections.cs
@@ -21,7 +21,7 @@ using System.Linq;
 
 using AgentInterface;
 using SwarmCoordinatorInterface;
-using CommonUtils;
+using SwarmCommonUtils;
 
 namespace Agent
 {
diff --git a/Engine/Source/Programs/UnrealSwarm/CommonUtils/CommonUtils.csproj b/Engine/Source/Programs/UnrealSwarm/CommonUtils/CommonUtils.csproj
deleted file mode 100644
diff --git a/Engine/Source/Programs/UnrealSwarm/CommonUtils/Properties/AssemblyInfo.cs b/Engine/Source/Programs/UnrealSwarm/CommonUtils/Properties/AssemblyInfo.cs
deleted file mode 100644
diff --git a/Engine/Source/Programs/UnrealSwarm/CommonUtils/Utils.cs b/Engine/Source/Programs/UnrealSwarm/CommonUtils/Utils.cs
deleted file mode 100644
diff --git a/Engine/Source/Programs/UnrealSwarm/SwarmCommonUtils/Properties/AssemblyInfo.cs b/Engine/Source/Programs/UnrealSwarm/SwarmCommonUtils/Properties/AssemblyInfo.cs
new file mode 100644
--- /dev/null
+++ b/Engine/Source/Programs/UnrealSwarm/SwarmCommonUtils/Properties/AssemblyInfo.cs
@@ -0,0 +1,38 @@
+﻿// Copyright 1998-2015 Epic Games, Inc. All Rights Reserved.
+
+using System.Reflection;
+using System.Runtime.CompilerServices;
+using System.Runtime.InteropServices;
+
+// General Information about an assembly is controlled through the following 
+// set of attributes. Change these attribute values to modify the information
+// associated with an assembly.
+[assembly: AssemblyTitle("SwarmCommonUtils")]
+[assembly: AssemblyDescription("")]
+[assembly: AssemblyConfiguration("")]
+[assembly: AssemblyCompany("")]
+[assembly: AssemblyProduct("SwarmCommonUtils")]
+[assembly: AssemblyCopyright("Copyright ©  2014")]
+[assembly: AssemblyTrademark("")]
+[assembly: AssemblyCulture("")]
+
+// Setting ComVisible to false makes the types in this assembly not visible 
+// to COM components.  If you need to access a type in this assembly from 
+// COM, set the ComVisible attribute to true on that type.
+[assembly: ComVisible(false)]
+
+// The following GUID is for the ID of the typelib if this project is exposed to COM
+[assembly: Guid("67079ec4-8039-4160-ac89-e5da5582dc8c")]
+
+// Version information for an assembly consists of the following four values:
+//
+//      Major Version
+//      Minor Version 
+//      Build Number
+//      Revision
+//
+// You can specify all the values or you can default the Build and Revision Numbers 
+// by using the '*' as shown below:
+// [assembly: AssemblyVersion("1.0.*")]
+[assembly: AssemblyVersion("1.0.0.0")]
+[assembly: AssemblyFileVersion("1.0.0.0")]
diff --git a/Engine/Source/Programs/UnrealSwarm/SwarmCommonUtils/SwarmCommonUtils.csproj b/Engine/Source/Programs/UnrealSwarm/SwarmCommonUtils/SwarmCommonUtils.csproj
new file mode 100644
--- /dev/null
+++ b/Engine/Source/Programs/UnrealSwarm/SwarmCommonUtils/SwarmCommonUtils.csproj
@@ -0,0 +1,54 @@
+﻿<?xml version="1.0" encoding="utf-8"?>
+<Project ToolsVersion="12.0" DefaultTargets="Build" xmlns="http://schemas.microsoft.com/developer/msbuild/2003">
+  <Import Project="$(MSBuildExtensionsPath)\$(MSBuildToolsVersion)\Microsoft.Common.props" Condition="Exists('$(MSBuildExtensionsPath)\$(MSBuildToolsVersion)\Microsoft.Common.props')" />
+  <PropertyGroup>
+    <Configuration Condition=" '$(Configuration)' == '' ">Debug</Configuration>
+    <Platform Condition=" '$(Platform)' == '' ">AnyCPU</Platform>
+    <ProjectGuid>{65521250-B022-442F-B961-3AF5ADA02083}</ProjectGuid>
+    <OutputType>Library</OutputType>
+    <AppDesignerFolder>Properties</AppDesignerFolder>
+    <RootNamespace>SwarmCommonUtils</RootNamespace>
+    <AssemblyName>SwarmCommonUtils</AssemblyName>
+    <TargetFrameworkVersion>v4.0</TargetFrameworkVersion>
+    <FileAlignment>512</FileAlignment>
+    <TargetFrameworkProfile>Client</TargetFrameworkProfile>
+  </PropertyGroup>
+  <PropertyGroup Condition=" '$(Configuration)|$(Platform)' == 'Debug|AnyCPU' ">
+    <DebugSymbols>true</DebugSymbols>
+    <DebugType>full</DebugType>
+    <Optimize>false</Optimize>
+    <OutputPath>..\..\..\..\Binaries\DotNET\</OutputPath>
+    <DefineConstants>DEBUG;TRACE</DefineConstants>
+    <ErrorReport>prompt</ErrorReport>
+    <WarningLevel>4</WarningLevel>
+  </PropertyGroup>
+  <PropertyGroup Condition=" '$(Configuration)|$(Platform)' == 'Development|AnyCPU' ">
+    <DebugType>pdbonly</DebugType>
+    <Optimize>true</Optimize>
+    <OutputPath>..\..\..\..\Binaries\DotNET\</OutputPath>
+    <DefineConstants>TRACE</DefineConstants>
+    <ErrorReport>prompt</ErrorReport>
+    <WarningLevel>4</WarningLevel>
+  </PropertyGroup>
+  <ItemGroup>
+    <Reference Include="System" />
+    <Reference Include="System.Core" />
+    <Reference Include="System.Xml.Linq" />
+    <Reference Include="System.Data.DataSetExtensions" />
+    <Reference Include="Microsoft.CSharp" />
+    <Reference Include="System.Data" />
+    <Reference Include="System.Xml" />
+  </ItemGroup>
+  <ItemGroup>
+    <Compile Include="Utils.cs" />
+    <Compile Include="Properties\AssemblyInfo.cs" />
+  </ItemGroup>
+  <Import Project="$(MSBuildToolsPath)\Microsoft.CSharp.targets" />
+  <!-- To modify your build process, add your task inside one of the targets below and uncomment it. 
+       Other similar extension points exist, see Microsoft.Common.targets.
+  <Target Name="BeforeBuild">
+  </Target>
+  <Target Name="AfterBuild">
+  </Target>
+  -->
+</Project>
\ No newline at end of file
diff --git a/Engine/Source/Programs/UnrealSwarm/SwarmCommonUtils/Utils.cs b/Engine/Source/Programs/UnrealSwarm/SwarmCommonUtils/Utils.cs
new file mode 100644
--- /dev/null
+++ b/Engine/Source/Programs/UnrealSwarm/SwarmCommonUtils/Utils.cs
@@ -0,0 +1,78 @@
+﻿// Copyright 1998-2015 Epic Games, Inc. All Rights Reserved.
+
+using System;
+using System.Collections.Generic;
+using System.IO;
+using System.Linq;
+using System.Net;
+using System.Net.NetworkInformation;
+using System.Runtime.InteropServices;
+using System.Runtime.Serialization.Formatters.Binary;
+using System.Text;
+using System.Threading.Tasks;
+
+namespace SwarmCommonUtils
+{
+	public static class NetworkUtils
+	{
+		[DllImport("iphlpapi.dll", SetLastError = true, EntryPoint = "GetBestInterface")]
+		static extern int WINAPI_GetBestInterface(UInt32 DestAddr, out UInt32 BestIfIndex);
+
+		public static NetworkInterface GetBestInterface(IPAddress Address)
+		{
+			UInt32 AddressAsUInt32 = BitConverter.ToUInt32(Address.GetAddressBytes(), 0);
+			UInt32 InterfaceIndex;
+
+			if (WINAPI_GetBestInterface(AddressAsUInt32, out InterfaceIndex) != 0)
+			{
+				throw new InvalidOperationException("Couldn't detect best network interface for that address.");
+			}
+
+			var Interfaces = NetworkInterface.GetAllNetworkInterfaces().Where(Interface =>
+			{
+				return Interface.Supports(NetworkInterfaceComponent.IPv4) && Interface.GetIPProperties().GetIPv4Properties().Index == InterfaceIndex;
+			});
+
+			if (Interfaces.Count() == 0)
+			{
+				throw new InvalidOperationException("Couldn't detect best network interface for that address.");
+			}
+
+			return Interfaces.First();
+		}
+
+		public static IPAddress GetInterfaceIPv4Address(NetworkInterface NetworkInterface)
+		{
+			var NetworkInterfaceProp = NetworkInterface.GetIPProperties();
+
+			foreach (var Unicast in NetworkInterfaceProp.UnicastAddresses)
+			{
+				if (Unicast.Address.AddressFamily == System.Net.Sockets.AddressFamily.InterNetwork)
+				{
+					return Unicast.Address;
+				}
+			}
+
+			throw new InvalidOperationException("Couldn't find interface IPv4 address.");
+		}
+	}
+
+	public static class ObjectUtils
+	{
+		public static T Duplicate<T>(T SourceObject)
+		{
+			if (!typeof(T).IsSerializable)
+			{
+				throw new InvalidOperationException("Only serializable types are supported.");
+			}
+
+			var Formatter = new BinaryFormatter();
+			using (var MemoryStream = new MemoryStream())
+			{
+				Formatter.Serialize(MemoryStream, SourceObject);
+				MemoryStream.Seek(0, SeekOrigin.Begin);
+				return (T)Formatter.Deserialize(MemoryStream);
+			}
+		}
+	}
+}
diff --git a/Engine/Source/Programs/UnrealSwarm/SwarmCoordinator/Coordinator.cs b/Engine/Source/Programs/UnrealSwarm/SwarmCoordinator/Coordinator.cs
--- a/Engine/Source/Programs/UnrealSwarm/SwarmCoordinator/Coordinator.cs
+++ b/Engine/Source/Programs/UnrealSwarm/SwarmCoordinator/Coordinator.cs
@@ -8,7 +8,7 @@ using System.Net;
 using System.Net.NetworkInformation;
 using System.Text;
 using SwarmCoordinatorInterface;
-using CommonUtils;
+using SwarmCommonUtils;
 
 namespace SwarmCoordinator
 {
diff --git a/Engine/Source/Programs/UnrealSwarm/SwarmCoordinator/SwarmCoordinator.csproj b/Engine/Source/Programs/UnrealSwarm/SwarmCoordinator/SwarmCoordinator.csproj
--- a/Engine/Source/Programs/UnrealSwarm/SwarmCoordinator/SwarmCoordinator.csproj
+++ b/Engine/Source/Programs/UnrealSwarm/SwarmCoordinator/SwarmCoordinator.csproj
@@ -121,9 +121,9 @@
     </Compile>
   </ItemGroup>
   <ItemGroup>
-    <ProjectReference Include="..\CommonUtils\CommonUtils.csproj">
+    <ProjectReference Include="..\SwarmCommonUtils\SwarmCommonUtils.csproj">
       <Project>{65521250-b022-442f-b961-3af5ada02083}</Project>
-      <Name>CommonUtils</Name>
+      <Name>SwarmCommonUtils</Name>
     </ProjectReference>
     <ProjectReference Include="..\SwarmCoordinatorInterface\SwarmCoordinatorInterface.csproj">
       <Project>{4C299DD0-B968-49F2-894A-9F115C369E72}</Project>
diff --git a/Engine/Source/Programs/UnrealSwarm/SwarmCoordinator/SwarmCoordinator_Mono.csproj b/Engine/Source/Programs/UnrealSwarm/SwarmCoordinator/SwarmCoordinator_Mono.csproj
--- a/Engine/Source/Programs/UnrealSwarm/SwarmCoordinator/SwarmCoordinator_Mono.csproj
+++ b/Engine/Source/Programs/UnrealSwarm/SwarmCoordinator/SwarmCoordinator_Mono.csproj
@@ -161,9 +161,9 @@
   </Target>
   -->
   <ItemGroup>
-    <ProjectReference Include="..\CommonUtils\CommonUtils.csproj">
+    <ProjectReference Include="..\SwarmCommonUtils\SwarmCommonUtils.csproj">
       <Project>{65521250-b022-442f-b961-3af5ada02083}</Project>
-      <Name>CommonUtils</Name>
+      <Name>SwarmCommonUtils</Name>
     </ProjectReference>
     <ProjectReference Include="..\SwarmCoordinatorInterface\SwarmCoordinatorInterface_Mono.csproj">
       <Project>{4C299DD0-B968-49F2-894A-9F115C369E72}</Project>
