UnrealSwarm: Swarm was not accepting host names in settings -- only IPs. Fixed.

diff --git a/Engine/Source/Programs/UnrealSwarm/Agent/Connections.cs b/Engine/Source/Programs/UnrealSwarm/Agent/Connections.cs
--- a/Engine/Source/Programs/UnrealSwarm/Agent/Connections.cs
+++ b/Engine/Source/Programs/UnrealSwarm/Agent/Connections.cs
@@ -2108,7 +2108,9 @@ namespace Agent
 						AgentInfoUpdate.Configuration["WorkingFor"] = WorkingFor;
 
 #if !__MonoCS__ // @todo Mac
-						var NetworkInterface = NetworkUtils.GetBestInterface(IPAddress.Parse(AgentApplication.Options.CoordinatorRemotingHost));
+						var NetworkInterface = NetworkUtils.GetBestInterface(
+							Dns.GetHostEntry(AgentApplication.Options.CoordinatorRemotingHost).AddressList[0]
+						);
 						AgentInfoUpdate.Configuration["IPAddress"] = NetworkUtils.GetInterfaceIPv4Address(NetworkInterface);
 #else
 						AgentInfoUpdate.Configuration["IPAddress"] = "192.168.0.203";//.2.9";
