Disable signing manifests for Swarm projects by default. Can be enabled manually if someone needs to publish, but currently cannot be built on any machine without the Epic cert installed.

diff --git a/Engine/Source/Programs/UnrealSwarm/Agent/Agent.csproj b/Engine/Source/Programs/UnrealSwarm/Agent/Agent.csproj
--- a/Engine/Source/Programs/UnrealSwarm/Agent/Agent.csproj
+++ b/Engine/Source/Programs/UnrealSwarm/Agent/Agent.csproj
@@ -19,7 +19,7 @@
     <ManifestKeyFile>
     </ManifestKeyFile>
     <GenerateManifests>true</GenerateManifests>
-    <SignManifests>true</SignManifests>
+    <SignManifests>false</SignManifests>
     <FileUpgradeFlags>
     </FileUpgradeFlags>
     <OldToolsVersion>3.5</OldToolsVersion>
diff --git a/Engine/Source/Programs/UnrealSwarm/SwarmCoordinator/SwarmCoordinator.csproj b/Engine/Source/Programs/UnrealSwarm/SwarmCoordinator/SwarmCoordinator.csproj
--- a/Engine/Source/Programs/UnrealSwarm/SwarmCoordinator/SwarmCoordinator.csproj
+++ b/Engine/Source/Programs/UnrealSwarm/SwarmCoordinator/SwarmCoordinator.csproj
@@ -22,7 +22,7 @@
     </UpgradeBackupLocation>
     <TargetFrameworkVersion>v4.0</TargetFrameworkVersion>
     <IsWebBootstrapper>false</IsWebBootstrapper>
-    <SignManifests>true</SignManifests>
+    <SignManifests>false</SignManifests>
     <SignAssembly>false</SignAssembly>
     <ManifestCertificateThumbprint>E10B5956C0FAFF8479C61A48C48EFDE555A18633</ManifestCertificateThumbprint>
     <ManifestKeyFile>
