UE-5796: Swarm agent log spam [PushChannel] ../Win64/UnrealLightmass.pdb has failed!

Changed log verbosity for those messages.

diff --git a/Engine/Source/Programs/UnrealSwarm/Agent/Channels.cs b/Engine/Source/Programs/UnrealSwarm/Agent/Channels.cs
--- a/Engine/Source/Programs/UnrealSwarm/Agent/Channels.cs
+++ b/Engine/Source/Programs/UnrealSwarm/Agent/Channels.cs
@@ -865,7 +865,7 @@ namespace Agent
 				}
 				else
 				{
-					Log(EVerbosityLevel.Informative, ELogColour.Red, string.Format("[PushChannel] Pushing the channel {0} has failed!", ChannelName));
+					Log(EVerbosityLevel.Verbose, ELogColour.Red, string.Format("[PushChannel] Pushing the channel {0} has failed!", ChannelName));
 				}
 			}
 			else
@@ -1007,7 +1007,7 @@ namespace Agent
 
 				if(!bChannelTransferred)
 				{
-					Log(EVerbosityLevel.Informative, ELogColour.Red, string.Format("[PullChannel] Pulling the channel {0} has failed! Retry {1} of {2}.", ChannelName, TryId + 1, RetriesOnFailure));
+					Log(EVerbosityLevel.Verbose, ELogColour.Red, string.Format("[PullChannel] Pulling the channel {0} has failed! Retry {1} of {2}.", ChannelName, TryId + 1, RetriesOnFailure));
 				}
 
 				++TryId;
