UE-16133: Swarm Agent Local Host of Coordinator not showing up in Coordinator

Fixed UnrealSwarm connection problems, which were happening if coordinator and agent were on the same machine and coordinator address in agent's settings was set to non-loopback IP.

diff --git a/Engine/Source/Programs/UnrealSwarm/Agent/Connections.cs b/Engine/Source/Programs/UnrealSwarm/Agent/Connections.cs
--- a/Engine/Source/Programs/UnrealSwarm/Agent/Connections.cs
+++ b/Engine/Source/Programs/UnrealSwarm/Agent/Connections.cs
@@ -22,6 +22,7 @@ using System.Linq;
 using AgentInterface;
 using SwarmCoordinatorInterface;
 using SwarmCommonUtils;
+using System.Net.Sockets;
 
 namespace Agent
 {
@@ -2108,10 +2109,21 @@ namespace Agent
 						AgentInfoUpdate.Configuration["WorkingFor"] = WorkingFor;
 
 #if !__MonoCS__ // @todo Mac
-						var NetworkInterface = NetworkUtils.GetBestInterface(
-							Dns.GetHostEntry(AgentApplication.Options.CoordinatorRemotingHost).AddressList[0]
-						);
-						AgentInfoUpdate.Configuration["IPAddress"] = NetworkUtils.GetInterfaceIPv4Address(NetworkInterface);
+						IPAddress[] CoordinatorAddresses = Dns.GetHostAddresses(AgentApplication.Options.CoordinatorRemotingHost);
+						IPAddress[] LocalAddresses = Dns.GetHostAddresses(Dns.GetHostName());
+
+						if(CoordinatorAddresses.Any(CoordinatorAddress => IPAddress.IsLoopback(CoordinatorAddress) || LocalAddresses.Contains(CoordinatorAddress)))
+						{
+							AgentInfoUpdate.Configuration["IPAddress"] = IPAddress.Loopback;
+						}
+						else
+						{
+							var NetworkInterface = NetworkUtils.GetBestInterface(
+								CoordinatorAddresses.Where(CoordinatorAddress => CoordinatorAddress.AddressFamily == AddressFamily.InterNetwork).First()
+							);
+
+							AgentInfoUpdate.Configuration["IPAddress"] = NetworkUtils.GetInterfaceIPv4Address(NetworkInterface);
+						}
 #else
 						AgentInfoUpdate.Configuration["IPAddress"] = "192.168.0.203";//.2.9";
 #endif
