Fixed UnrealSwarm "Channel already registered" error to log the actual reason why exception has happened, cause until now it was very misleading.

diff --git a/Engine/Source/Programs/UnrealSwarm/Agent/AgentApplication.cs b/Engine/Source/Programs/UnrealSwarm/Agent/AgentApplication.cs
--- a/Engine/Source/Programs/UnrealSwarm/Agent/AgentApplication.cs
+++ b/Engine/Source/Programs/UnrealSwarm/Agent/AgentApplication.cs
@@ -137,13 +137,19 @@ namespace Agent
 						ChannelServices.RegisterChannel( AgentTCPChannel, false );
 					}
 				}
-				catch( Exception )
+				catch (RemotingException Ex)
 				{
-					AgentApplication.Log( EVerbosityLevel.Informative, ELogColour.Orange, "[ERROR] Channel already registered, suggesting another SwarmAgent or client is running." );
-					AgentApplication.Log( EVerbosityLevel.Informative, ELogColour.Orange, "[ERROR] If you feel this is in error, check your running process list for additional copies of" );
-					AgentApplication.Log( EVerbosityLevel.Informative, ELogColour.Orange, "[ERROR] SwarmAgent or UnrealLightmass (or other client) and consider killing them." );
-					AgentApplication.Log( EVerbosityLevel.Informative, ELogColour.Orange, "[ERROR] Sleeping for a few seconds and trying again..." );
-					Thread.Sleep( 3000 );
+					AgentApplication.Log(EVerbosityLevel.Informative, ELogColour.Orange, "[ERROR] Channel already registered, suggesting another SwarmAgent or client is running.");
+					AgentApplication.Log(EVerbosityLevel.Informative, ELogColour.Orange, "[ERROR] If you feel this is in error, check your running process list for additional copies of");
+					AgentApplication.Log(EVerbosityLevel.Informative, ELogColour.Orange, "[ERROR] SwarmAgent or UnrealLightmass (or other client) and consider killing them.");
+					AgentApplication.Log(EVerbosityLevel.Informative, ELogColour.Orange, "[ERROR] Sleeping for a few seconds and trying again...");
+					AgentApplication.Log(EVerbosityLevel.Informative, ELogColour.Orange, string.Format("[ERROR] Channel registration failed. Reason: {0}\n, Callstack: {1}.", Ex.Message, Ex.StackTrace));
+					Thread.Sleep(3000);
+				}
+				catch (Exception Ex)
+				{
+					AgentApplication.Log(EVerbosityLevel.Informative, ELogColour.Orange, string.Format("[ERROR] Channel registration failed. Reason: {0}\n, Callstack: {1}.", Ex.Message, Ex.StackTrace));
+					Thread.Sleep(3000);
 				}
 			}
 
