UE-12935: Lightmass crash when building lighting in RenderTestMap

Could not find a repro, but I've added additional logging and retry functionality in the place this could occur on network failure to have more data next time it happens.

diff --git a/Engine/Source/Programs/UnrealSwarm/Agent/Connections.cs b/Engine/Source/Programs/UnrealSwarm/Agent/Connections.cs
--- a/Engine/Source/Programs/UnrealSwarm/Agent/Connections.cs
+++ b/Engine/Source/Programs/UnrealSwarm/Agent/Connections.cs
@@ -526,20 +526,26 @@ namespace Agent
 			bool ReturnCode = false;
 			if( RemoteInterfaceAlive )
 			{
-				try
-				{
-					RequestChannelDelegate DRequestChannel = new RequestChannelDelegate( RemoteInterface.RequestChannel );
-					IAsyncResult Result = DRequestChannel.BeginInvoke( ConnectionHandle, ChannelName, JobGuid, null, null );
-					WaitHandle.WaitAny( new WaitHandle[2] { Result.AsyncWaitHandle, RemoteInterfaceDropped } );
-					if( Result.IsCompleted )
+				RetryCatch(
+					// Try {
+					() =>
 					{
-						ReturnCode = DRequestChannel.EndInvoke( Result );
-					}
-				}
-				catch( Exception )
-				{
-					SignalConnectionDropped();
-				}
+						RequestChannelDelegate DRequestChannel = new RequestChannelDelegate(RemoteInterface.RequestChannel);
+						IAsyncResult Result = DRequestChannel.BeginInvoke(ConnectionHandle, ChannelName, JobGuid, null, null);
+						WaitHandle.WaitAny(new WaitHandle[2] { Result.AsyncWaitHandle, RemoteInterfaceDropped });
+						if (Result.IsCompleted)
+						{
+							ReturnCode = DRequestChannel.EndInvoke(Result);
+						}
+					},
+					// } Catch (Exception) {
+					() =>
+					{
+						SignalConnectionDropped();
+					},
+					// }
+					5
+				);
 			}
 			return ReturnCode;
 		}
@@ -688,6 +694,33 @@ namespace Agent
 			return ReturnCode;
 		}
 
+		private void RetryCatch(Action TryClause, Action CatchClause, int RetriesOnFailure = 5)
+		{
+			int TryId = 1;
+			while(TryId <= RetriesOnFailure)
+			{
+				try
+				{
+					TryClause();
+					break;
+				}
+				catch(Exception Ex)
+				{
+					if(TryId == RetriesOnFailure)
+					{
+						AgentApplication.Log(EVerbosityLevel.Critical, ELogColour.Red, string.Format("Error (retry {0} of {1}): {2}. Stack trace:\n{3}\n", TryId, RetriesOnFailure, Ex.Message, Ex.StackTrace));
+						CatchClause();
+					}
+					else
+					{
+						AgentApplication.Log(EVerbosityLevel.Critical, ELogColour.Orange, string.Format("Error (retry {0} of {1}): {2}. Stack trace: {3}.\n", TryId, RetriesOnFailure, Ex.Message, Ex.StackTrace));
+					}
+				}
+
+				++TryId;
+			}
+		}
+
 		public string				RemoteAgentName;
 		public string				RemoteAgentIPAddress;
 		private Int32				RemoteAgentHandle;
diff --git a/Engine/Source/Programs/UnrealSwarm/Agent/Properties/AssemblyInfo.cs b/Engine/Source/Programs/UnrealSwarm/Agent/Properties/AssemblyInfo.cs
--- a/Engine/Source/Programs/UnrealSwarm/Agent/Properties/AssemblyInfo.cs
+++ b/Engine/Source/Programs/UnrealSwarm/Agent/Properties/AssemblyInfo.cs
@@ -17,4 +17,4 @@ using System.Reflection;
 //      Build Number
 //      Revision
 //
-[assembly: AssemblyVersion( "1.7.2892.0" )]
+[assembly: AssemblyVersion( "1.7.2893.0" )]
