Back out native linux UAT support - waiting for proper fix from dev-platform.

diff --git a/Engine/Source/Programs/IOS/MobileDeviceInterface/MobileDeviceInterface.csproj b/Engine/Source/Programs/IOS/MobileDeviceInterface/MobileDeviceInterface.csproj
--- a/Engine/Source/Programs/IOS/MobileDeviceInterface/MobileDeviceInterface.csproj
+++ b/Engine/Source/Programs/IOS/MobileDeviceInterface/MobileDeviceInterface.csproj
@@ -10,7 +10,7 @@
     <AppDesignerFolder>Properties</AppDesignerFolder>
     <RootNamespace>MobileDeviceInterface</RootNamespace>
     <AssemblyName>MobileDeviceInterface</AssemblyName>
-    <TargetFrameworkVersion>v4.5</TargetFrameworkVersion>
+    <TargetFrameworkVersion>v4.0</TargetFrameworkVersion>
     <FileAlignment>512</FileAlignment>
     <SccProjectName>
     </SccProjectName>
diff --git a/Engine/Source/Programs/NetworkProfiler/NetworkProfiler/NetworkProfiler.csproj b/Engine/Source/Programs/NetworkProfiler/NetworkProfiler/NetworkProfiler.csproj
--- a/Engine/Source/Programs/NetworkProfiler/NetworkProfiler/NetworkProfiler.csproj
+++ b/Engine/Source/Programs/NetworkProfiler/NetworkProfiler/NetworkProfiler.csproj
@@ -10,7 +10,7 @@
     <AppDesignerFolder>Properties</AppDesignerFolder>
     <RootNamespace>NetworkProfiler</RootNamespace>
     <AssemblyName>NetworkProfiler</AssemblyName>
-    <TargetFrameworkVersion>v4.5</TargetFrameworkVersion>
+    <TargetFrameworkVersion>v4.0</TargetFrameworkVersion>
     <FileAlignment>512</FileAlignment>
     <SccProjectName>
     </SccProjectName>
diff --git a/Engine/Source/Programs/UnrealControls/UnrealControls.csproj b/Engine/Source/Programs/UnrealControls/UnrealControls.csproj
--- a/Engine/Source/Programs/UnrealControls/UnrealControls.csproj
+++ b/Engine/Source/Programs/UnrealControls/UnrealControls.csproj
@@ -19,7 +19,7 @@
     <OldToolsVersion>3.5</OldToolsVersion>
     <UpgradeBackupLocation>
     </UpgradeBackupLocation>
-    <TargetFrameworkVersion>v4.5</TargetFrameworkVersion>
+    <TargetFrameworkVersion>v4.0</TargetFrameworkVersion>
     <IsWebBootstrapper>false</IsWebBootstrapper>
     <TargetFrameworkProfile>Client</TargetFrameworkProfile>
     <PublishUrl>publish\</PublishUrl>
diff --git a/Engine/Source/Programs/UnrealSwarm/SwarmCommonUtils/SwarmCommonUtils.csproj b/Engine/Source/Programs/UnrealSwarm/SwarmCommonUtils/SwarmCommonUtils.csproj
--- a/Engine/Source/Programs/UnrealSwarm/SwarmCommonUtils/SwarmCommonUtils.csproj
+++ b/Engine/Source/Programs/UnrealSwarm/SwarmCommonUtils/SwarmCommonUtils.csproj
@@ -9,7 +9,7 @@
     <AppDesignerFolder>Properties</AppDesignerFolder>
     <RootNamespace>SwarmCommonUtils</RootNamespace>
     <AssemblyName>SwarmCommonUtils</AssemblyName>
-    <TargetFrameworkVersion>v4.5</TargetFrameworkVersion>
+    <TargetFrameworkVersion>v4.0</TargetFrameworkVersion>
     <FileAlignment>512</FileAlignment>
     <TargetFrameworkProfile>Client</TargetFrameworkProfile>
   </PropertyGroup>
diff --git a/Engine/Source/Programs/UnrealSwarm/SwarmCoordinator/SwarmCoordinator.csproj b/Engine/Source/Programs/UnrealSwarm/SwarmCoordinator/SwarmCoordinator.csproj
--- a/Engine/Source/Programs/UnrealSwarm/SwarmCoordinator/SwarmCoordinator.csproj
+++ b/Engine/Source/Programs/UnrealSwarm/SwarmCoordinator/SwarmCoordinator.csproj
@@ -20,7 +20,7 @@
     <OldToolsVersion>3.5</OldToolsVersion>
     <UpgradeBackupLocation>
     </UpgradeBackupLocation>
-    <TargetFrameworkVersion>v4.5</TargetFrameworkVersion>
+    <TargetFrameworkVersion>v4.0</TargetFrameworkVersion>
     <IsWebBootstrapper>false</IsWebBootstrapper>
     <SignManifests>false</SignManifests>
     <SignAssembly>false</SignAssembly>
@@ -80,7 +80,7 @@
     <Reference Include="System.Management" />
     <Reference Include="System.Runtime.Remoting" />
     <Reference Include="System.Windows.Forms" />
-    <Reference Include="System.XML" />
+    <Reference Include="System.Xml" />
   </ItemGroup>
   <ItemGroup>
     <Compile Include="..\..\DotNETCommon\MetaData.cs">
diff --git a/Engine/Source/Programs/UnrealSwarm/SwarmCoordinatorInterface/SwarmCoordinatorInterface.csproj b/Engine/Source/Programs/UnrealSwarm/SwarmCoordinatorInterface/SwarmCoordinatorInterface.csproj
--- a/Engine/Source/Programs/UnrealSwarm/SwarmCoordinatorInterface/SwarmCoordinatorInterface.csproj
+++ b/Engine/Source/Programs/UnrealSwarm/SwarmCoordinatorInterface/SwarmCoordinatorInterface.csproj
@@ -20,7 +20,7 @@
     <OldToolsVersion>3.5</OldToolsVersion>
     <UpgradeBackupLocation>
     </UpgradeBackupLocation>
-    <TargetFrameworkVersion>v4.5</TargetFrameworkVersion>
+    <TargetFrameworkVersion>v4.0</TargetFrameworkVersion>
     <PublishUrl>publish\</PublishUrl>
     <Install>true</Install>
     <InstallFrom>Disk</InstallFrom>
