Changing coordinator publish directory to /app.publish2 to stop a race condition in clickonce deployment that can cause the agent to wipe out the coordinator binaries or vice versa

diff --git a/Engine/Source/Programs/UnrealSwarm/SwarmCoordinator/SwarmCoordinator.csproj b/Engine/Source/Programs/UnrealSwarm/SwarmCoordinator/SwarmCoordinator.csproj
--- a/Engine/Source/Programs/UnrealSwarm/SwarmCoordinator/SwarmCoordinator.csproj
+++ b/Engine/Source/Programs/UnrealSwarm/SwarmCoordinator/SwarmCoordinator.csproj
@@ -1,5 +1,5 @@
 ﻿<?xml version="1.0" encoding="utf-8"?>
-<Project DefaultTargets="Build" xmlns="http://schemas.microsoft.com/developer/msbuild/2003" ToolsVersion="4.0">
+<Project DefaultTargets="Build;Clean" xmlns="http://schemas.microsoft.com/developer/msbuild/2003" ToolsVersion="4.0">
   <PropertyGroup>
     <Configuration Condition=" '$(Configuration)' == '' ">Debug</Configuration>
     <Platform Condition=" '$(Platform)' == '' ">AnyCPU</Platform>
@@ -59,6 +59,7 @@
     <DebugType>pdbonly</DebugType>
     <Optimize>false</Optimize>
     <OutputPath>..\..\..\..\Binaries\DotNET\</OutputPath>
+    <PublishDir>$(OutputPath)app.publish2\</PublishDir>
     <DefineConstants>DEBUG;TRACE</DefineConstants>
     <ErrorReport>prompt</ErrorReport>
     <WarningLevel>4</WarningLevel>
@@ -68,6 +69,7 @@
     <DebugType>pdbonly</DebugType>
     <Optimize>true</Optimize>
     <OutputPath>..\..\..\..\Binaries\DotNET\</OutputPath>
+    <PublishDir>$(OutputPath)app.publish2\</PublishDir>
     <DefineConstants>TRACE</DefineConstants>
     <ErrorReport>prompt</ErrorReport>
     <WarningLevel>4</WarningLevel>
@@ -175,8 +177,11 @@
   <!-- To modify your build process, add your task inside one of the targets below and uncomment it. 
        Other similar extension points exist, see Microsoft.Common.targets.
   <Target Name="BeforeBuild">
-  </Target>
+  </Target>-->
   <Target Name="AfterBuild">
   </Target>
-  -->
+  <Target Name="CleanPublishDir" AfterTargets="Clean">
+    <!-- When $(PublishDir) is overridden we need to clean it manually (see Microsoft.Common.CurrentVersion.targets) -->
+    <RemoveDir Directories="$(PublishDir)" />
+  </Target>
 </Project>
\ No newline at end of file
