- Reverted the Unreal.MachineName change in Swarm Coordinator where it shouldn't have been made, but replacing with System.Net.Dns.GetHostName()

diff --git a/Engine/Source/Programs/UnrealSwarm/SwarmCoordinatorInterface/SwarmCoordinatorInterface.cs b/Engine/Source/Programs/UnrealSwarm/SwarmCoordinatorInterface/SwarmCoordinatorInterface.cs
--- a/Engine/Source/Programs/UnrealSwarm/SwarmCoordinatorInterface/SwarmCoordinatorInterface.cs
+++ b/Engine/Source/Programs/UnrealSwarm/SwarmCoordinatorInterface/SwarmCoordinatorInterface.cs
@@ -32,7 +32,7 @@ namespace SwarmCoordinatorInterface
     public class AgentInfo
     {
 		// Constant agent attributes
-		public string Name = Unreal.MachineName;
+		public string Name = System.Net.Dns.GetHostName();
 		public AgentState State = AgentState.Unknown;
 		public Version Version = new Version( 0, 0, 0, 0 );
 		public Hashtable Configuration = new Hashtable();
