- Fixed up the MachineName/Dns stuff like SwarmCoordinatorInterface.cs

diff --git a/Engine/Source/Programs/UnrealSwarm/SwarmCoordinator/Coordinator.cs b/Engine/Source/Programs/UnrealSwarm/SwarmCoordinator/Coordinator.cs
--- a/Engine/Source/Programs/UnrealSwarm/SwarmCoordinator/Coordinator.cs
+++ b/Engine/Source/Programs/UnrealSwarm/SwarmCoordinator/Coordinator.cs
@@ -128,7 +128,7 @@ namespace SwarmCoordinator
 				UpdatedInfo.Configuration["IPAddress"] = IPAddress.Loopback;
 				bool IPAddressUpdateSuccess = false;
 
-				if( UpdatedInfo.Name != Unreal.MachineName)
+				if( UpdatedInfo.Name != System.Net.Dns.GetHostName())
 				{
 					try
 					{
diff --git a/Engine/Source/Programs/UnrealSwarm/SwarmCoordinator/SwarmCoordinator.cs b/Engine/Source/Programs/UnrealSwarm/SwarmCoordinator/SwarmCoordinator.cs
--- a/Engine/Source/Programs/UnrealSwarm/SwarmCoordinator/SwarmCoordinator.cs
+++ b/Engine/Source/Programs/UnrealSwarm/SwarmCoordinator/SwarmCoordinator.cs
@@ -55,7 +55,7 @@ namespace SwarmCoordinator
             Coordinator.Init();
 
 			// Set the title bar to include the name of the machine
-			TopLevelControl.Text = "Swarm Coordinator running on " + Unreal.MachineName;
+			TopLevelControl.Text = "Swarm Coordinator running on " + System.Net.Dns.GetHostName();
 
             // Display the window
             Show();
