- More fixes for Unreal.MachineName

diff --git a/Engine/Source/Programs/IOS/iPhonePackager/CompileTime.cs b/Engine/Source/Programs/IOS/iPhonePackager/CompileTime.cs
--- a/Engine/Source/Programs/IOS/iPhonePackager/CompileTime.cs
+++ b/Engine/Source/Programs/IOS/iPhonePackager/CompileTime.cs
@@ -107,7 +107,7 @@ namespace iPhonePackager
 		 */
 		static public void ConfigurePaths()
 		{
-			string MachineName = Unreal.MachineName;
+			string MachineName = System.Net.Dns.GetHostName();
 
 			XcodeDeveloperDir = Utilities.GetEnvironmentVariable("ue.XcodeDeveloperDir", "/Applications/Xcode.app/Contents/Developer/");
 
diff --git a/Engine/Source/Programs/UnrealControls/GameLocator.cs b/Engine/Source/Programs/UnrealControls/GameLocator.cs
--- a/Engine/Source/Programs/UnrealControls/GameLocator.cs
+++ b/Engine/Source/Programs/UnrealControls/GameLocator.cs
@@ -112,9 +112,9 @@ namespace UnrealControls
 			}
 
 			// Add in the machine name
-			if(Unreal.MachineName.Length > 0 )
+			if(System.Net.Dns.GetHostName().Length > 0 )
 			{
-				DefaultTargetDirectory += "-" + Unreal.MachineName;
+				DefaultTargetDirectory += "-" + System.Net.Dns.GetHostName();
 			}
 
 			return ( DefaultTargetDirectory );
diff --git a/Engine/Source/Programs/UnrealSwarm/Agent/Connections.cs b/Engine/Source/Programs/UnrealSwarm/Agent/Connections.cs
--- a/Engine/Source/Programs/UnrealSwarm/Agent/Connections.cs
+++ b/Engine/Source/Programs/UnrealSwarm/Agent/Connections.cs
@@ -1112,7 +1112,7 @@ namespace Agent
 			{
 				// Update our state and ping the coordinator, if this is an instigating connection
 				CurrentState = AgentState.Working;
-				WorkingFor = Unreal.MachineName;
+				WorkingFor = System.Net.Dns.GetHostName();
 				PingCoordinator( true );
 			}
 
@@ -1227,7 +1227,7 @@ namespace Agent
 
 						Hashtable RequestedConfiguration = new Hashtable();
 						RequestedConfiguration["Version"] = new Version( 0, 0, 0, 0 );
-						RequestedConfiguration["RequestingAgentName"] = Unreal.MachineName;
+						RequestedConfiguration["RequestingAgentName"] = System.Net.Dns.GetHostName();
 						List<AgentInfo> PotentialRemoteAgents = Coordinator.GetAvailableAgents( RequestedConfiguration );
 						foreach( AgentInfo NextAgentInfo in PotentialRemoteAgents )
 						{
@@ -1484,7 +1484,7 @@ namespace Agent
 					// Send the connection handle and the machine name that will
 					// be used to construct the return URL for bi-directional
 					// connection
-					if( WrappedRemoteAgentInterface.OpenRemoteConnection(Unreal.MachineName, NewConnectionHandle, LoggingFlags, IsAssigned, AssignedTimestamp ) >= 0 )
+					if( WrappedRemoteAgentInterface.OpenRemoteConnection(System.Net.Dns.GetHostName(), NewConnectionHandle, LoggingFlags, IsAssigned, AssignedTimestamp ) >= 0 )
 					{
 						// Successful confirmation, double check that the connection has moved
 						if( PendingConnections.ContainsKey( NewConnectionHandle ) == false )
@@ -1802,8 +1802,8 @@ namespace Agent
 					Hashtable RequestedConfiguration = new Hashtable();
 					RequestedConfiguration["Version"] = CurrentVersion;
 					RequestedConfiguration["GroupName"] = AgentApplication.Options.AllowedRemoteAgentGroup;
-					RequestedConfiguration["RequestingAgentName"] = Unreal.MachineName;
-					RequestedConfiguration["RequestAssignmentFor"] = Unreal.MachineName;
+					RequestedConfiguration["RequestingAgentName"] = System.Net.Dns.GetHostName();
+					RequestedConfiguration["RequestAssignmentFor"] = System.Net.Dns.GetHostName();
 					PotentialRemoteAgents = Coordinator.GetAvailableAgents( RequestedConfiguration );
 				}
 				catch( Exception )
@@ -1820,7 +1820,7 @@ namespace Agent
 					foreach( AgentInfo Info in PotentialRemoteAgents )
 					{
 						bool WorkerAlreadyConnectedOrDisallowed = false;
-						if( ( Info.Name == Unreal.MachineName) ||
+						if( ( Info.Name == System.Net.Dns.GetHostName()) ||
 							( !AgentNamePassesAllowedAgentsFilter( Info.Name ) ) )
 						{
 							// Trivial case
@@ -1882,7 +1882,7 @@ namespace Agent
 				bool IsAssigned = false;
 				DateTime AssignedTimestamp = DateTime.MinValue;
 				if( ( NextRemoteAgent.Configuration.ContainsKey( "AssignedTo" ) ) &&
-					( ( NextRemoteAgent.Configuration["AssignedTo"] as string ) == Unreal.MachineName) )
+					( ( NextRemoteAgent.Configuration["AssignedTo"] as string ) == System.Net.Dns.GetHostName()) )
 				{
 					IsAssigned = true;
 					AssignedTimestamp = ( DateTime )NextRemoteAgent.Configuration["AssignedTime"];
@@ -2015,7 +2015,7 @@ namespace Agent
 					Hashtable RequestedConfiguration = new Hashtable();
 					RequestedConfiguration["Version"] = CurrentVersion;
 					RequestedConfiguration["GroupName"] = AgentGroupName;
-					RequestedConfiguration["RequestingAgentName"] = Unreal.MachineName;
+					RequestedConfiguration["RequestingAgentName"] = System.Net.Dns.GetHostName();
 					PotentialRemoteAgents = Coordinator.GetAvailableAgents( RequestedConfiguration );
 					PrintSimpleAgentList( PotentialRemoteAgents );
 				}
@@ -2043,7 +2043,7 @@ namespace Agent
 					Hashtable RequestedConfiguration = new Hashtable();
 					RequestedConfiguration["Version"] = new Version( 0, 0, 0, 0 );
 					RequestedConfiguration["GroupName"] = AgentGroupName;
-					RequestedConfiguration["RequestingAgentName"] = Unreal.MachineName;
+					RequestedConfiguration["RequestingAgentName"] = System.Net.Dns.GetHostName();
 
 					// Start by getting the current state of the agents in the group we care about
 					List<AgentInfo> RemoteAgentsInGroupBefore = Coordinator.GetAvailableAgents( RequestedConfiguration );
diff --git a/Engine/Source/Programs/UnrealSwarm/Agent/Display.cs b/Engine/Source/Programs/UnrealSwarm/Agent/Display.cs
--- a/Engine/Source/Programs/UnrealSwarm/Agent/Display.cs
+++ b/Engine/Source/Programs/UnrealSwarm/Agent/Display.cs
@@ -86,7 +86,7 @@ namespace Agent
 			UpdateWindowState();
 
 			// Set the title bar to include the name of the machine and the group
-			TopLevelControl.Text = "Swarm Agent running on " + Unreal.MachineName;
+			TopLevelControl.Text = "Swarm Agent running on " + System.Net.Dns.GetHostName();
 
 			LogOutputWindowView.Refresh();
 		}
diff --git a/Engine/Source/Programs/UnrealSwarm/Agent/Jobs.cs b/Engine/Source/Programs/UnrealSwarm/Agent/Jobs.cs
--- a/Engine/Source/Programs/UnrealSwarm/Agent/Jobs.cs
+++ b/Engine/Source/Programs/UnrealSwarm/Agent/Jobs.cs
@@ -91,7 +91,7 @@ namespace Agent
 				Manager.SetCurrentDirectoryByProcessID( LocalNewOwner.ProcessID );
 
 				// Update the visualizer
-				AgentApplication.UpdateMachineState(Unreal.MachineName, -1, EProgressionState.InstigatorConnected );
+				AgentApplication.UpdateMachineState(System.Net.Dns.GetHostName(), -1, EProgressionState.InstigatorConnected );
 			}
 
 			// Add to the Agent-wide list of active jobs
@@ -443,7 +443,7 @@ namespace Agent
 						string SqlCommandString = "EXEC dbo.AddJob_v1";
 						SqlCommandString += String.Format( " @Duration='{0}',", ( StopTime - StartTime ).TotalSeconds );
 						SqlCommandString += String.Format( " @UserName='{0}',", Environment.UserName.ToUpperInvariant() );
-						SqlCommandString += String.Format( " @MachineName='{0}',", Unreal.MachineName.ToUpperInvariant() );
+						SqlCommandString += String.Format( " @MachineName='{0}',", System.Net.Dns.GetHostName().ToUpperInvariant() );
 						SqlCommandString += String.Format( " @GroupName='{0}',", AgentApplication.Options.AgentGroupName );
 						SqlCommandString += String.Format( " @JobGUID='{0}',", JobGuid.ToString() );
 						SqlCommandString += String.Format( " @Instigator='{0}',", OwnerIsInstigator );
@@ -999,7 +999,7 @@ namespace Agent
 						string WorkerAgentIPAddress = LastSuccessfulJobRecord.WorkerAgentIPAddresses[i];
 
 						// Don't try to connect to self and only if this name is allowed
-						if( ( WorkerAgentName != Unreal.MachineName) &&
+						if( ( WorkerAgentName != System.Net.Dns.GetHostName()) &&
 							( Manager.AgentNamePassesAllowedAgentsFilter( WorkerAgentName ) ) )
 						{
 							// Make sure we can open the connection and get a valid remote connection back
@@ -1139,7 +1139,7 @@ namespace Agent
 						if( OwnerIsInstigator )
 						{
 							// Send the total task count to the visualizer if this agent is the Instigator
-							AgentApplication.UpdateMachineState(Unreal.MachineName, ( Int32 )TaskCount, EProgressionState.TaskTotal );
+							AgentApplication.UpdateMachineState(System.Net.Dns.GetHostName(), ( Int32 )TaskCount, EProgressionState.TaskTotal );
 						}
 
 						// If we have more than N pending tasks, the job owner is the instigator,
@@ -1260,7 +1260,7 @@ namespace Agent
 					if( OwnerIsInstigator )
 					{
 						// Inform the visualizer that we've disconnected
-						AgentApplication.UpdateMachineState(Unreal.MachineName, -1, EProgressionState.InstigatorDisconnected );
+						AgentApplication.UpdateMachineState(System.Net.Dns.GetHostName(), -1, EProgressionState.InstigatorDisconnected );
 
 						// If the job was a success, record the state for determinisitc replay
 						if( CurrentSuccessState == JobSuccessState.AGENT_JOB_SUCCESS )
@@ -1712,7 +1712,7 @@ namespace Agent
 						}
 						else
 						{
-							Manager.Log( EVerbosityLevel.Informative, ELogColour.Red, "[UpdateTaskState]: Task Rejected locally by " + Unreal.MachineName + ", counted as failure");
+							Manager.Log( EVerbosityLevel.Informative, ELogColour.Red, "[UpdateTaskState]: Task Rejected locally by " + System.Net.Dns.GetHostName() + ", counted as failure");
 							UpdateTaskStateAsFailure( RunningTask );
 						}
 						break;
@@ -1726,7 +1726,7 @@ namespace Agent
 						}
 						else
 						{
-							Manager.Log( EVerbosityLevel.Informative, ELogColour.Red, "[UpdateTaskState]: Task Killed locally by " + Unreal.MachineName + ", counted as failure" );
+							Manager.Log( EVerbosityLevel.Informative, ELogColour.Red, "[UpdateTaskState]: Task Killed locally by " + System.Net.Dns.GetHostName() + ", counted as failure" );
 							UpdateTaskStateAsFailure( RunningTask );
 						}
 						break;
@@ -1738,7 +1738,7 @@ namespace Agent
 						}
 						else
 						{
-							Manager.Log( EVerbosityLevel.Informative, ELogColour.Red, "[UpdateTaskState]: Task Failed on " + Unreal.MachineName);
+							Manager.Log( EVerbosityLevel.Informative, ELogColour.Red, "[UpdateTaskState]: Task Failed on " + System.Net.Dns.GetHostName());
 						}
 						Manager.Log( EVerbosityLevel.Informative, ELogColour.Red, "[UpdateTaskState]: Task Failed: " + RunningTask.Specification.Parameters );
 						UpdateTaskStateAsFailure( RunningTask );
@@ -1783,8 +1783,8 @@ namespace Agent
 				// Update the visualizer if this agent is the Instigator
 				if( OwnerIsInstigator )
 				{
-					AgentApplication.UpdateMachineState(Unreal.MachineName, RetiredTasks.Count, EProgressionState.TasksCompleted );
-					AgentApplication.UpdateMachineState(Unreal.MachineName, RunningTasks.Count, EProgressionState.TasksInProgress );
+					AgentApplication.UpdateMachineState(System.Net.Dns.GetHostName(), RetiredTasks.Count, EProgressionState.TasksCompleted );
+					AgentApplication.UpdateMachineState(System.Net.Dns.GetHostName(), RunningTasks.Count, EProgressionState.TasksInProgress );
 				}
 			}
 		}
diff --git a/Engine/Source/Programs/UnrealSwarm/Agent/Messaging.cs b/Engine/Source/Programs/UnrealSwarm/Agent/Messaging.cs
--- a/Engine/Source/Programs/UnrealSwarm/Agent/Messaging.cs
+++ b/Engine/Source/Programs/UnrealSwarm/Agent/Messaging.cs
@@ -285,7 +285,7 @@ namespace Agent
 			RemoteConnection Remote = RequestingConnection as RemoteConnection;
 			if( Remote == null )
 			{
-				Name = Unreal.MachineName;
+				Name = System.Net.Dns.GetHostName();
 			}
 			else
 			{
diff --git a/Engine/Source/Programs/UnrealSwarm/Agent/Progression.cs b/Engine/Source/Programs/UnrealSwarm/Agent/Progression.cs
--- a/Engine/Source/Programs/UnrealSwarm/Agent/Progression.cs
+++ b/Engine/Source/Programs/UnrealSwarm/Agent/Progression.cs
@@ -596,9 +596,9 @@ namespace Agent
 				}
 				
 				// Always move the local machine name to the top of the list
-				if( MachineNames.Remove(Unreal.MachineName) )
+				if( MachineNames.Remove(System.Net.Dns.GetHostName()) )
 				{
-					MachineNames.Insert( 0, Unreal.MachineName);
+					MachineNames.Insert( 0, System.Net.Dns.GetHostName());
 				}
 
 				// Find ideal row height dependent on number of agents
diff --git a/Engine/Source/Programs/UnrealSwarm/Agent/Utils.cs b/Engine/Source/Programs/UnrealSwarm/Agent/Utils.cs
--- a/Engine/Source/Programs/UnrealSwarm/Agent/Utils.cs
+++ b/Engine/Source/Programs/UnrealSwarm/Agent/Utils.cs
@@ -463,7 +463,7 @@ namespace Agent
 
         public string DumpTimings()
         {
-            string Output = "Machine: " + Unreal.MachineName + Environment.NewLine + Environment.NewLine;
+            string Output = "Machine: " + System.Net.Dns.GetHostName() + Environment.NewLine + Environment.NewLine;
 
     		double TotalTime = 0.0;
 	    	foreach( PerfTiming Timing in Timings.Values )
